#ue4
#messaging
* temporary fix for accessing the messaging module from a thread other than the game thread

diff --git a/Engine/Source/Runtime/Messaging/Public/Interfaces/IMessagingModule.h b/Engine/Source/Runtime/Messaging/Public/Interfaces/IMessagingModule.h
--- a/Engine/Source/Runtime/Messaging/Public/Interfaces/IMessagingModule.h
+++ b/Engine/Source/Runtime/Messaging/Public/Interfaces/IMessagingModule.h
@@ -68,7 +68,12 @@ public:
 	 */
 	static IMessagingModule& Get( )
 	{
-		return FModuleManager::LoadModuleChecked<IMessagingModule>("Messaging");
+#if PLATFORM_IOS
+        static IMessagingModule& MessageModule = FModuleManager::LoadModuleChecked<IMessagingModule>("Messaging");
+        return MessageModule;
+#else
+        return FModuleManager::LoadModuleChecked<IMessagingModule>("Messaging");
+#endif
 	}
 
 public:
