Messaging: documentation & nullptr cleanup

diff --git a/Engine/Source/Runtime/Messaging/Private/Bridge/MessageAddressBook.h b/Engine/Source/Runtime/Messaging/Private/Bridge/MessageAddressBook.h
--- a/Engine/Source/Runtime/Messaging/Private/Bridge/MessageAddressBook.h
+++ b/Engine/Source/Runtime/Messaging/Private/Bridge/MessageAddressBook.h
@@ -35,8 +35,8 @@ public:
 	/**
 	 * Adds an address to the address book.
 	 *
-	 * @Address - The message address to add.
-	 * @NodeId - The identifier of the transport node that handles the message address.
+	 * @Address The message address to add.
+	 * @NodeId The identifier of the transport node that handles the message address.
 	 */
 	void Add( const FMessageAddress& Address, const FGuid& NodeId )
 	{
@@ -58,7 +58,7 @@ public:
 	/**
 	 * Checks whether this address book contains the given address.
 	 *
-	 * @param Address - The address to check.
+	 * @param Address The address to check.
 	 *
 	 * @return true if the address is known, false otherwise.
 	 */
@@ -72,7 +72,7 @@ public:
 	/**
 	 * Gets the transport node identifiers for the specified list of message addresses.
 	 *
-	 * @param Addresses - The address list to retrieve the node identifiers for.
+	 * @param Addresses The address list to retrieve the node identifiers for.
 	 *
 	 * @return The list of node identifiers.
 	 */
@@ -101,7 +101,7 @@ public:
 	 * To remove only the addresses for a specific transport node, use RemoveNode().
 	 * If you are not interested in the removed addresses, use Clear().
 	 *
-	 * @param OutRemovedRecipients - Will hold a list of recipients that were removed.
+	 * @param OutRemovedRecipients Will hold a list of recipients that were removed.
 	 *
 	 * @see Clear
 	 * @see RemoveNode
@@ -119,8 +119,8 @@ public:
 	/**
 	 * Removes all known message addresses for the specified transport node identifier.
 	 *
-	 * @param NodeId - The identifier of the transport node to remove.
-	 * @param OutRemovedRecipients - Will hold a list of recipients that were removed.
+	 * @param NodeId The identifier of the transport node to remove.
+	 * @param OutRemovedRecipients Will hold a list of recipients that were removed.
 	 *
 	 * @see Clear
 	 * @see RemoveAllNodes
diff --git a/Engine/Source/Runtime/Messaging/Private/Bridge/MessageBridge.h b/Engine/Source/Runtime/Messaging/Private/Bridge/MessageBridge.h
--- a/Engine/Source/Runtime/Messaging/Private/Bridge/MessageBridge.h
+++ b/Engine/Source/Runtime/Messaging/Private/Bridge/MessageBridge.h
@@ -21,10 +21,10 @@ public:
 	/**
 	 * Creates and initializes a new instance.
 	 *
-	 * @param InAddress - The address for this bridge.
-	 * @param InBus - The message bus that this node is connected to.
-	 * @param InSerializer - The message serializer to use.
-	 * @param InTransport - The transport mechanism to use.
+	 * @param InAddress The address for this bridge.
+	 * @param InBus The message bus that this node is connected to.
+	 * @param InSerializer The message serializer to use.
+	 * @param InTransport The transport mechanism to use.
 	 */
 	FMessageBridge( const FMessageAddress InAddress, const IMessageBusRef& InBus, const ISerializeMessagesRef& InSerializer, const ITransportMessagesRef& InTransport );
 
diff --git a/Engine/Source/Runtime/Messaging/Private/Bridge/MessageData.h b/Engine/Source/Runtime/Messaging/Private/Bridge/MessageData.h
--- a/Engine/Source/Runtime/Messaging/Private/Bridge/MessageData.h
+++ b/Engine/Source/Runtime/Messaging/Private/Bridge/MessageData.h
@@ -40,7 +40,7 @@ public:
 	/**
 	 * Updates the state of this message data.
 	 *
-	 * @param InState - The state to set.
+	 * @param InState The state to set.
 	 */
 	void UpdateState( EMessageDataState::Type InState )
 	{
diff --git a/Engine/Source/Runtime/Messaging/Private/Bridge/MessageDeserializeTask.h b/Engine/Source/Runtime/Messaging/Private/Bridge/MessageDeserializeTask.h
--- a/Engine/Source/Runtime/Messaging/Private/Bridge/MessageDeserializeTask.h
+++ b/Engine/Source/Runtime/Messaging/Private/Bridge/MessageDeserializeTask.h
@@ -17,9 +17,9 @@ public:
 	/**
 	 * Creates and initializes a new instance.
 	 *
-	 * InMessageTypeInfo - The message's type information (may be nullptr if the type information hasn't been discovered yet).
-	 * @param InSerializer - The serializer to use.
-	 * @param InThread - The name of the thread to deserialize the message on.
+	 * InMessageTypeInfo The message's type information (may be nullptr if the type information hasn't been discovered yet).
+	 * @param InSerializer The serializer to use.
+	 * @param InThread The name of the thread to deserialize the message on.
 	 */
 	FMessageDeserializeTask( TWeakObjectPtr<UScriptStruct> InMessageTypeInfo, ISerializeMessagesRef InSerializer, ENamedThreads::Type InThread )
 		: MessageTypeInfoPtr(InMessageTypeInfo)
@@ -32,8 +32,8 @@ public:
 	/**
 	 * Performs the actual task.
 	 *
-	 * @param CurrentThread - The thread that this task is executing on.
-	 * @param MyCompletionGraphEvent - The completion event.
+	 * @param CurrentThread The thread that this task is executing on.
+	 * @param MyCompletionGraphEvent The completion event.
 	 *
 	 * @todo gmp: implement asynchronous message deserialization
 	 */
diff --git a/Engine/Source/Runtime/Messaging/Private/Bridge/MessageForwardTask.h b/Engine/Source/Runtime/Messaging/Private/Bridge/MessageForwardTask.h
--- a/Engine/Source/Runtime/Messaging/Private/Bridge/MessageForwardTask.h
+++ b/Engine/Source/Runtime/Messaging/Private/Bridge/MessageForwardTask.h
@@ -17,9 +17,9 @@ public:
 	/**
 	 * Creates and initializes a new instance.
 	 *
-	 * @param InMessageContext - The context of the message to forward.
-	 * @param InBus - The message bus to forward the message to.
-	 * @param InBridge - The sender that is forwarding the message.
+	 * @param InMessageContext The context of the message to forward.
+	 * @param InBus The message bus to forward the message to.
+	 * @param InBridge The sender that is forwarding the message.
 	 */
 	FMessageForwardTask( IMessageContextRef InMessageContext, IMessageBusRef InBus, ISendMessagesRef InSender )
 		: BusPtr(InBus)
@@ -32,8 +32,8 @@ public:
 	/**
 	 * Performs the actual task.
 	 *
-	 * @param CurrentThread - The thread that this task is executing on.
-	 * @param MyCompletionGraphEvent - The completion event.
+	 * @param CurrentThread The thread that this task is executing on.
+	 * @param MyCompletionGraphEvent The completion event.
 	 */
 	void DoTask( ENamedThreads::Type CurrentThread, const FGraphEventRef& MyCompletionGraphEvent )
 	{
diff --git a/Engine/Source/Runtime/Messaging/Private/Bridge/MessageSerializeTask.h b/Engine/Source/Runtime/Messaging/Private/Bridge/MessageSerializeTask.h
--- a/Engine/Source/Runtime/Messaging/Private/Bridge/MessageSerializeTask.h
+++ b/Engine/Source/Runtime/Messaging/Private/Bridge/MessageSerializeTask.h
@@ -17,9 +17,9 @@ public:
 	/**
 	 * Creates and initializes a new instance.
 	 *
-	 * @param InMessageContext - The context of the message to serialize.
-	 * @param InMessageData - Will hold the serialized message data.
-	 * @param InSerializer - The serializer to use.
+	 * @param InMessageContext The context of the message to serialize.
+	 * @param InMessageData Will hold the serialized message data.
+	 * @param InSerializer The serializer to use.
 	 */
 	FMessageSerializeTask( IMessageContextRef InMessageContext, FMessageDataRef InMessageData, ISerializeMessagesRef InSerializer )
 		: MessageContext(InMessageContext)
@@ -32,8 +32,8 @@ public:
 	/**
 	 * Performs the actual task.
 	 *
-	 * @param CurrentThread - The thread that this task is executing on.
-	 * @param MyCompletionGraphEvent - The completion event.
+	 * @param CurrentThread The thread that this task is executing on.
+	 * @param MyCompletionGraphEvent The completion event.
 	 */
 	void DoTask( ENamedThreads::Type CurrentThread, const FGraphEventRef& MyCompletionGraphEvent )
 	{
diff --git a/Engine/Source/Runtime/Messaging/Private/Bus/MessageBus.h b/Engine/Source/Runtime/Messaging/Private/Bus/MessageBus.h
--- a/Engine/Source/Runtime/Messaging/Private/Bus/MessageBus.h
+++ b/Engine/Source/Runtime/Messaging/Private/Bus/MessageBus.h
@@ -19,7 +19,7 @@ public:
 	/**
 	 * Creates and initializes a new instance.
 	 *
-	 * @param InRecipientAuthorizer - An optional recipient authorizer.
+	 * @param InRecipientAuthorizer An optional recipient authorizer.
 	 */
 	FMessageBus( const IAuthorizeMessageRecipientsPtr& InRecipientAuthorizer );
 
diff --git a/Engine/Source/Runtime/Messaging/Private/Bus/MessageContext.h b/Engine/Source/Runtime/Messaging/Private/Bus/MessageContext.h
--- a/Engine/Source/Runtime/Messaging/Private/Bus/MessageContext.h
+++ b/Engine/Source/Runtime/Messaging/Private/Bus/MessageContext.h
@@ -29,15 +29,15 @@ public:
 	/**
 	 * Creates and initializes a new message context.
 	 *
-	 * @param InMessage - The message payload.
-	 * @param InTypeInfo - The message's type information.
-	 * @param InAttachment - The binary data to attach to the message.
-	 * @param InSender - The sender's address.
-	 * @param InRecipients - The message recipients.
-	 * @param InScope - The message scope.
-	 * @param InTimeSent - The time at which the message was sent.
-	 * @param InExpiration - The message's expiration time.
-	 * @param InSenderThread - The name of the thread from which the message was sent.
+	 * @param InMessage The message payload.
+	 * @param InTypeInfo The message's type information.
+	 * @param InAttachment The binary data to attach to the message.
+	 * @param InSender The sender's address.
+	 * @param InRecipients The message recipients.
+	 * @param InScope The message scope.
+	 * @param InTimeSent The time at which the message was sent.
+	 * @param InExpiration The message's expiration time.
+	 * @param InSenderThread The name of the thread from which the message was sent.
 	 */
 	FMessageContext( void* InMessage, UScriptStruct* InTypeInfo, const IMessageAttachmentPtr& InAttachment, const FMessageAddress& InSender, const TArray<FMessageAddress>& InRecipients, EMessageScope::Type InScope, const FDateTime& InTimeSent, const FDateTime& InExpiration, ENamedThreads::Type InSenderThread )
 		: Attachment(InAttachment)
@@ -54,12 +54,12 @@ public:
 	/**
 	 * Creates and initializes a new message from an existing message context.
 	 *
-	 * @param InContext - The context to initialize from.
-	 * @param InForwarder - The forwarder's address.
-	 * @param NewRecipients - The message recipients.
-	 * @param NewScope - The forwarding scope.
-	 * @param InTimeForwarded - The time at which the message was forwarded.
-	 * @param InSenderThread - The name of the thread from which the message was sent.
+	 * @param InContext The context to initialize from.
+	 * @param InForwarder The forwarder's address.
+	 * @param NewRecipients The message recipients.
+	 * @param NewScope The forwarding scope.
+	 * @param InTimeForwarded The time at which the message was forwarded.
+	 * @param InSenderThread The name of the thread from which the message was sent.
 	 */
 	FMessageContext( const IMessageContextRef& InContext, const FMessageAddress& InForwarder, const TArray<FMessageAddress>& NewRecipients, EMessageScope::Type NewScope, const FDateTime& InTimeForwarded, ENamedThreads::Type InSenderThread )
 		: Attachment(nullptr)
diff --git a/Engine/Source/Runtime/Messaging/Private/Bus/MessageDispatchTask.h b/Engine/Source/Runtime/Messaging/Private/Bus/MessageDispatchTask.h
--- a/Engine/Source/Runtime/Messaging/Private/Bus/MessageDispatchTask.h
+++ b/Engine/Source/Runtime/Messaging/Private/Bus/MessageDispatchTask.h
@@ -17,10 +17,10 @@ public:
 	/**
 	 * Creates and initializes a new instance.
 	 *
-	 * @param InThread - The name of the thread to dispatch the message on.
-	 * @param InContext - The context of the message to dispatch.
-	 * @param InRecipient - The message recipient.
-	 * @param InTracer - The message tracer to notify.
+	 * @param InThread The name of the thread to dispatch the message on.
+	 * @param InContext The context of the message to dispatch.
+	 * @param InRecipient The message recipient.
+	 * @param InTracer The message tracer to notify.
 	 */
 	FMessageDispatchTask( ENamedThreads::Type InThread, IMessageContextRef InContext, IReceiveMessagesWeakPtr InRecipient, FMessageTracerPtr InTracer )
 		: Context(InContext)
@@ -34,8 +34,8 @@ public:
 	/**
 	 * Performs the actual task.
 	 *
-	 * @param CurrentThread - The thread that this task is executing on.
-	 * @param MyCompletionGraphEvent - The completion event.
+	 * @param CurrentThread The thread that this task is executing on.
+	 * @param MyCompletionGraphEvent The completion event.
 	 */
 	void DoTask( ENamedThreads::Type CurrentThread, const FGraphEventRef& MyCompletionGraphEvent )
 	{
diff --git a/Engine/Source/Runtime/Messaging/Private/Bus/MessageRouter.h b/Engine/Source/Runtime/Messaging/Private/Bus/MessageRouter.h
--- a/Engine/Source/Runtime/Messaging/Private/Bus/MessageRouter.h
+++ b/Engine/Source/Runtime/Messaging/Private/Bus/MessageRouter.h
@@ -30,8 +30,8 @@ public:
 	/**
 	 * Adds a message interceptor.
 	 *
-	 * @param Interceptor - The interceptor to add.
-	 * @param MessageType - The type of messages to intercept.
+	 * @param Interceptor The interceptor to add.
+	 * @param MessageType The type of messages to intercept.
 	 */
 	FORCEINLINE void AddInterceptor( const IInterceptMessagesRef& Interceptor, const FName& MessageType )
 	{
@@ -41,8 +41,8 @@ public:
 	/**
 	 * Adds a recipient.
 	 *
-	 * @param Address - The address of the recipient to add.
-	 * @param Recipient - The recipient.
+	 * @param Address The address of the recipient to add.
+	 * @param Recipient The recipient.
 	 */
 	FORCEINLINE void AddRecipient( const FMessageAddress& Address, const IReceiveMessagesRef& Recipient )
 	{
@@ -52,7 +52,7 @@ public:
 	/**
 	 * Adds a subscription.
 	 *
-	 * @param Subscription - The subscription to add.
+	 * @param Subscription The subscription to add.
 	 */
 	FORCEINLINE void AddSubscription( const IMessageSubscriptionRef& Subscription )
 	{
@@ -72,8 +72,8 @@ public:
 	/**
 	 * Removes a message interceptor.
 	 *
-	 * @param Interceptor - The interceptor to remove.
-	 * @param MessageType - The type of messages to stop intercepting.
+	 * @param Interceptor The interceptor to remove.
+	 * @param MessageType The type of messages to stop intercepting.
 	 */
 	FORCEINLINE void RemoveInterceptor( const IInterceptMessagesRef& Interceptor, const FName& MessageType )
 	{
@@ -83,7 +83,7 @@ public:
 	/**
 	 * Removes a recipient.
 	 *
-	 * @param Address - The address of the recipient to remove.
+	 * @param Address The address of the recipient to remove.
 	 */
 	FORCEINLINE void RemoveRecipient( const FMessageAddress& Address )
 	{
@@ -93,8 +93,8 @@ public:
 	/**
 	 * Removes a subscription.
 	 *
-	 * @param Subscriber - The subscriber to stop routing messages to.
-	 * @param MessageType - The type of message to unsubscribe from (NAME_None = all types).
+	 * @param Subscriber The subscriber to stop routing messages to.
+	 * @param MessageType The type of message to unsubscribe from (NAME_None = all types).
 	 */
 	FORCEINLINE void RemoveSubscription( const IReceiveMessagesRef& Subscriber, const FName& MessageType )
 	{
@@ -104,7 +104,7 @@ public:
 	/**
 	 * Routes a message to the specified recipients.
 	 *
-	 * @param Context - The context of the message to route.
+	 * @param Context The context of the message to route.
 	 */
 	FORCEINLINE void RouteMessage( const IMessageContextRef& Context )
 	{
@@ -138,7 +138,7 @@ protected:
 	/**
 	 * Queues up a router command.
 	 *
-	 * @param Command - The command to queue up.
+	 * @param Command The command to queue up.
 	 *
 	 * @return true if the command was enqueued, false otherwise.
 	 */
@@ -157,17 +157,17 @@ protected:
 	/**
 	 * Filters a collection of subscriptions using the given message context.
 	 *
-	 * @param Subscriptions - The subscriptions to filter.
-	 * @param Context - The message context to filter by.
-	 * @param Sender - The message sender (may be nullptr if the sender has no subscriptions).
-	 * @param OutRecipients - Will hold the collection of recipients.
+	 * @param Subscriptions The subscriptions to filter.
+	 * @param Context The message context to filter by.
+	 * @param Sender The message sender (may be nullptr if the sender has no subscriptions).
+	 * @param OutRecipients Will hold the collection of recipients.
 	 */
 	void FilterSubscriptions( TArray<IMessageSubscriptionPtr>& Subscriptions, const IMessageContextRef& Context, TArray<IReceiveMessagesPtr>& OutRecipients );
 
 	/**
 	 * Dispatches a single message to its recipients.
 	 *
-	 * @param Message - The message to dispatch.
+	 * @param Message The message to dispatch.
 	 */
 	void DispatchMessage( const IMessageContextRef& Message );
 
diff --git a/Engine/Source/Runtime/Messaging/Private/Bus/MessageSubscription.h b/Engine/Source/Runtime/Messaging/Private/Bus/MessageSubscription.h
--- a/Engine/Source/Runtime/Messaging/Private/Bus/MessageSubscription.h
+++ b/Engine/Source/Runtime/Messaging/Private/Bus/MessageSubscription.h
@@ -22,10 +22,10 @@ public:
 	/**
 	 * Creates and initializes a new instance.
 	 *
-	 * @param InSubscriber - The message subscriber.
-	 * @param InMessageType - The type of messages to subscribe to.
-	 * @param InReceivingThread - The thread on which to receive messages on.
-	 * @param InScopeRange - The message scope range to subscribe to.
+	 * @param InSubscriber The message subscriber.
+	 * @param InMessageType The type of messages to subscribe to.
+	 * @param InReceivingThread The thread on which to receive messages on.
+	 * @param InScopeRange The message scope range to subscribe to.
 	 */
 	FMessageSubscription( const IReceiveMessagesRef& InSubscriber, const FName& InMessageType, const FMessageScopeRange& InScopeRange )
 		: Enabled(true)
diff --git a/Engine/Source/Runtime/Messaging/Private/Bus/MessageTracer.h b/Engine/Source/Runtime/Messaging/Private/Bus/MessageTracer.h
--- a/Engine/Source/Runtime/Messaging/Private/Bus/MessageTracer.h
+++ b/Engine/Source/Runtime/Messaging/Private/Bus/MessageTracer.h
@@ -66,8 +66,8 @@ public:
 	/**
 	 * Notifies the tracer that a message interceptor has been added to the message bus.
 	 *
-	 * @param Interceptor - The added interceptor.
-	 * @param MessageType - The type of messages being intercepted.
+	 * @param Interceptor The added interceptor.
+	 * @param MessageType The type of messages being intercepted.
 	 */
 	void TraceAddedInterceptor( const IInterceptMessagesRef& Interceptor, const FName& MessageType )
 	{
@@ -78,8 +78,8 @@ public:
 	/**
 	 * Notifies the tracer that a message recipient has been added to the message bus.
 	 *
-	 * @param Address - The address of the added recipient.
-	 * @param Recipient - The added recipient.
+	 * @param Address The address of the added recipient.
+	 * @param Recipient The added recipient.
 	 */
 	void TraceAddedRecipient( const FMessageAddress& Address, const IReceiveMessagesRef& Recipient )
 	{
@@ -89,7 +89,7 @@ public:
 	/**
 	 * Notifies the tracer that a message subscription has been added to the message bus.
 	 *
-	 * @param Subscription - The added subscription.
+	 * @param Subscription The added subscription.
 	 */
 	void TraceAddedSubscription( const IMessageSubscriptionRef& Subscription )
 	{
@@ -103,9 +103,9 @@ public:
 	/**
 	 * Notifies the tracer that a message has been dispatched.
 	 *
-	 * @param Context - The context of the dispatched message.
-	 * @param Recipient - The message recipient.
-	 * @param Async - Whether the message was dispatched asynchronously.
+	 * @param Context The context of the dispatched message.
+	 * @param Recipient The message recipient.
+	 * @param Async Whether the message was dispatched asynchronously.
 	 */
 	void TraceDispatchedMessage( const IMessageContextRef& Context, const IReceiveMessagesRef& Recipient, bool Async )
 	{
@@ -120,8 +120,8 @@ public:
 	/**
 	 * Notifies the tracer that a message has been handled.
 	 *
-	 * @param Context - The context of the dispatched message.
-	 * @param Recipient - The message recipient that handled the message.
+	 * @param Context The context of the dispatched message.
+	 * @param Recipient The message recipient that handled the message.
 	 */
 	void TraceHandledMessage( const IMessageContextRef& Context, const IReceiveMessagesRef& Recipient )
 	{
@@ -136,8 +136,8 @@ public:
 	/**
 	 * Notifies the tracer that a message has been intercepted.
 	 *
-	 * @param Context - The context of the intercepted message.
-	 * @param Interceptor - The interceptor.
+	 * @param Context The context of the intercepted message.
+	 * @param Interceptor The interceptor.
 	 */
 	void TraceInterceptedMessage( const IMessageContextRef& Context, const IInterceptMessagesRef& Interceptor )
 	{
@@ -150,8 +150,8 @@ public:
 	/**
 	 * Notifies the tracer that a message interceptor has been removed from the message bus.
 	 *
-	 * @param Interceptor - The removed interceptor.
-	 * @param MessageType - The type of messages that is no longer being intercepted.
+	 * @param Interceptor The removed interceptor.
+	 * @param MessageType The type of messages that is no longer being intercepted.
 	 */
 	void TraceRemovedInterceptor( const IInterceptMessagesRef& Interceptor, const FName& MessageType )
 	{
@@ -165,7 +165,7 @@ public:
 	/**
 	 * Notifies the tracer that a recipient has been removed from the message bus.
 	 *
-	 * @param Address - The address of the removed recipient.
+	 * @param Address The address of the removed recipient.
 	 */
 	void TraceRemovedRecipient( const FMessageAddress& Address )
 	{
@@ -179,8 +179,8 @@ public:
 	/**
 	 * Notifies the tracer that a message subscription has been removed from the message bus.
 	 *
-	 * @param Subscriber - The removed subscriber.
-	 * @param MessageType - The type of messages no longer being subscribed to.
+	 * @param Subscriber The removed subscriber.
+	 * @param MessageType The type of messages no longer being subscribed to.
 	 */
 	void TraceRemovedSubscription( const IMessageSubscriptionRef& Subscription, const FName& MessageType )
 	{
@@ -194,7 +194,7 @@ public:
 	/**
 	 * Notifies the tracer that a message has been routed.
 	 *
-	 * @param Context - The context of the routed message.
+	 * @param Context The context of the routed message.
 	 */
 	void TraceRoutedMessage( const IMessageContextRef& Context )
 	{
@@ -215,7 +215,7 @@ public:
 	/**
 	 * Notifies the tracer that a message has been sent.
 	 *
-	 * @param Context - The context of the sent message.
+	 * @param Context The context of the sent message.
 	 */
 	void TraceSentMessage( const IMessageContextRef& Context )
 	{
@@ -333,7 +333,7 @@ protected:
 	/**
 	 * Enqueues a trace action for synchronized processing.
 	 *
-	 * @param Trace - The action to enqueue.
+	 * @param Trace The action to enqueue.
 	 */
 	FORCEINLINE void EnqueueTrace( TBaseDelegate_NoParams<void> Trace )
 	{
@@ -343,85 +343,85 @@ protected:
 	/**
 	 * Processes traces for added message interceptors.
 	 *
-	 * @param Name - The name of the interceptor.
-	 * @param MessageType - The type of messages being intercepted.
-	 * @param TimeSecond - The time at which the interceptor was added.
+	 * @param Name The name of the interceptor.
+	 * @param MessageType The type of messages being intercepted.
+	 * @param TimeSecond The time at which the interceptor was added.
 	 */
 	void ProcessAddedInterceptor( FString Name, FName MessageType, double TimeSeconds );
 
 	/**
 	 * Processes traces for added message recipients.
 	 *
-	 * @param Address - The address of the added recipient.
-	 * @param RecipientInfo - Information about the recipient (name, ID, etc.)
-	 * @param TimeSecond - The time at which the recipient was added.
+	 * @param Address The address of the added recipient.
+	 * @param RecipientInfo Information about the recipient (name, ID, etc.)
+	 * @param TimeSecond The time at which the recipient was added.
 	 */
 	void ProcessAddedRecipient( FMessageAddress Address, FRecipientInfo RecipientInfo, double TimeSeconds );
 
 	/**
 	 * Processes traces for added message subscriptions.
 	 *
-	 * @param TimeSecond - The time at which the subscription was added.
+	 * @param TimeSecond The time at which the subscription was added.
 	 */
 	void ProcessAddedSubscriptionTrace( double TimeSeconds );
 
 	/**
 	 * Processes traces for dispatched messages.
 	 *
-	 * @param Context - The context of the dispatched message.
-	 * @param TimeSecond - The time at which the message was dispatched.
-	 * @param RecipientId - The recipient's unique identifier.
-	 * @param Async - Whether the message was dispatched asynchronously.
+	 * @param Context The context of the dispatched message.
+	 * @param TimeSecond The time at which the message was dispatched.
+	 * @param RecipientId The recipient's unique identifier.
+	 * @param Async Whether the message was dispatched asynchronously.
 	 */
 	void ProcessDispatchedMessage( IMessageContextRef Context, double TimeSeconds, FGuid RecipientId, bool Async );
 
 	/**
 	 * Processes traces for handled messages.
 	 *
-	 * @param Context - The context of the handled message.
-	 * @param TimeSecond - The time at which the message was handled.
-	 * @param RecipientId - The recipient's unique identifier.
+	 * @param Context The context of the handled message.
+	 * @param TimeSecond The time at which the message was handled.
+	 * @param RecipientId The recipient's unique identifier.
 	 */
 	void ProcessHandledMessage( IMessageContextRef Context, double TimeSeconds, FGuid RecipientId );
 
 	/**
 	 * Processes traces for removed message interceptors.
 	 *
-	 * @param Interceptor - The removed interceptor.
-	 * @param MessageType - The type of messages no longer being intercepted.
-	 * @param TimeSecond - The time at which the interceptor was removed.
+	 * @param Interceptor The removed interceptor.
+	 * @param MessageType The type of messages no longer being intercepted.
+	 * @param TimeSecond The time at which the interceptor was removed.
 	 */
 	void ProcessRemovedInterceptor( IInterceptMessagesRef Interceptor, FName MessageType, double TimeSeconds );
 
 	/**
 	 * Processes traces for removed message recipients.
 	 *
-	 * @param Address - The address of the removed recipient.
-	 * @param TimeSecond - The time at which the recipient was removed.
+	 * @param Address The address of the removed recipient.
+	 * @param TimeSecond The time at which the recipient was removed.
 	 */
 	void ProcessRemovedRecipient( FMessageAddress Address, double TimeSeconds );
 
 	/**
 	 * Processes traces for removed message subscriptions.
 	 *
-	 * @param MessageType - The type of messages no longer being subscribed to.
-	 * @param TimeSecond - The time at which the subscription was removed.
+	 * @param MessageType The type of messages no longer being subscribed to.
+	 * @param TimeSecond The time at which the subscription was removed.
 	 */
 	void ProcessRemovedSubscription( FName MessageType, double TimeSeconds );
 
 	/**
 	 * Processes traces for routed messages.
 	 *
-	 * @param Context - The context of the routed message.
-	 * @param TimeSecond - The time at which the message was routed.
+	 * @param Context The context of the routed message.
+	 * @param TimeSecond The time at which the message was routed.
 	 */
 	void ProcessRoutedMessage( IMessageContextRef Context, double TimeSeconds );
 
 	/**
 	 * Processes traces for sent messages.
 	 *
-	 * @param Context - The context of the sent message.
-	 * @param TimeSecond - The time at which the message was sent.
+	 * @param Context The context of the sent message.
+	 * @param TimeSecond The time at which the message was sent.
 	 */
 	void ProcessSentMessage( IMessageContextRef Context, double TimeSeconds );
 
@@ -433,7 +433,7 @@ protected:
 	/**
 	 * Checks whether the tracer should break on the given message.
 	 *
-	 * @param Context - The context of the message to consider for breaking.
+	 * @param Context The context of the message to consider for breaking.
 	 */
 	bool ShouldBreak( const IMessageContextRef& Context ) const;
 
diff --git a/Engine/Source/Runtime/Messaging/Private/MessagingPrivatePCH.h b/Engine/Source/Runtime/Messaging/Private/MessagingPrivatePCH.h
--- a/Engine/Source/Runtime/Messaging/Private/MessagingPrivatePCH.h
+++ b/Engine/Source/Runtime/Messaging/Private/MessagingPrivatePCH.h
@@ -4,14 +4,13 @@
 	MessagingPrivatePCH.h: Pre-compiled header file for the Messaging module.
 =============================================================================*/
 
-#ifndef UNREAL_MESSAGING_PRIVATEPCH_H
-#define UNREAL_MESSAGING_PRIVATEPCH_H
+#pragma once
 
 
 #include "../Public/Messaging.h"
 
 
-/* Dependencies
+/* Private Dependencies
  *****************************************************************************/
 
 #include "Json.h"
@@ -39,6 +38,3 @@
 #include "MessageForwardTask.h"
 #include "MessageSerializeTask.h"
 #include "MessageBridge.h"
-
-
-#endif
diff --git a/Engine/Source/Runtime/Messaging/Private/Serialization/JsonMessageSerializer.cpp b/Engine/Source/Runtime/Messaging/Private/Serialization/JsonMessageSerializer.cpp
--- a/Engine/Source/Runtime/Messaging/Private/Serialization/JsonMessageSerializer.cpp
+++ b/Engine/Source/Runtime/Messaging/Private/Serialization/JsonMessageSerializer.cpp
@@ -35,7 +35,7 @@ struct FJsonStructWriteState
 };
 
 
-// @todo gmp - refactor this and SlateStyleSerializer into generic UStruct serializer
+// @todo gmp: refactor this and SlateStyleSerializer into generic UStruct serializer
 namespace Obsolete
 {
 	// Finds the class for the given stack state.
diff --git a/Engine/Source/Runtime/Messaging/Private/Serialization/JsonMessageSerializer.h b/Engine/Source/Runtime/Messaging/Private/Serialization/JsonMessageSerializer.h
--- a/Engine/Source/Runtime/Messaging/Private/Serialization/JsonMessageSerializer.h
+++ b/Engine/Source/Runtime/Messaging/Private/Serialization/JsonMessageSerializer.h
@@ -23,7 +23,6 @@ public:
 
 	// End ISerializeMessages interface
 
-
 protected:
 
 	/**
@@ -35,9 +34,9 @@ protected:
 	/**
 	 * Serializes a structure.
 	 *
-	 * @param Data - The structured data to serialize.
-	 * @param TypeInfo - The data's type information.
-	 * @param Writer - The Json writer to write to.
+	 * @param Data The structured data to serialize.
+	 * @param TypeInfo The data's type information.
+	 * @param Writer The Json writer to write to.
 	 */
 	bool SerializeStruct( const void* Data, UStruct& TypeInfo, const TSharedRef<TJsonWriter<UCS2CHAR> >& Writer );
 };
diff --git a/Engine/Source/Runtime/Messaging/Public/Common/FileMessageAttachment.h b/Engine/Source/Runtime/Messaging/Public/Common/FileMessageAttachment.h
--- a/Engine/Source/Runtime/Messaging/Public/Common/FileMessageAttachment.h
+++ b/Engine/Source/Runtime/Messaging/Public/Common/FileMessageAttachment.h
@@ -18,7 +18,7 @@ public:
 	/**
 	 * Creates and initializes a new instance.
 	 *
-	 * @param Filename - The full name and path of the file.
+	 * @param Filename The full name and path of the file.
 	 */
 	FFileMessageAttachment( const FString& InFilename )
 		: AutoDeleteFile(false)
@@ -28,8 +28,8 @@ public:
 	/**
 	 * Creates and initializes a new instance.
 	 *
-	 * @param Filename - The full name and path of the file.
-	 * @param AutoDeleteFile - Whether to delete the file when this attachment is destroyed.
+	 * @param Filename The full name and path of the file.
+	 * @param AutoDeleteFile Whether to delete the file when this attachment is destroyed.
 	 */
 	FFileMessageAttachment( const FString& InFilename, bool InAutoDeleteFile )
 		: AutoDeleteFile(InAutoDeleteFile)
diff --git a/Engine/Source/Runtime/Messaging/Public/Common/MessageBridgeBuilder.h b/Engine/Source/Runtime/Messaging/Public/Common/MessageBridgeBuilder.h
--- a/Engine/Source/Runtime/Messaging/Public/Common/MessageBridgeBuilder.h
+++ b/Engine/Source/Runtime/Messaging/Public/Common/MessageBridgeBuilder.h
@@ -21,19 +21,19 @@ public:
 		: Address(FMessageAddress::NewGuid())
 		, BusPtr(IMessagingModule::Get().GetDefaultBus())
 		, Disabled(false)
-		, Transport(NULL)
+		, Transport(nullptr)
 	{ }
 
 	/*
 	 * Creates and initializes a new instance.
 	 *
-	 * @param InBus - The message bus to attach the bridge to.
+	 * @param InBus The message bus to attach the bridge to.
 	 */
 	FMessageBridgeBuilder( const IMessageBusRef& Bus )
 		: Address(FMessageAddress::NewGuid())
 		, BusPtr(Bus)
 		, Disabled(false)
-		, Transport(NULL)
+		, Transport(nullptr)
 	{ }
 
 public:
@@ -53,7 +53,7 @@ public:
 	/**
 	 * Configures the bridge to use a custom message serializer.\
 	 *
-	 * @param CustomSerializer - The custom serializer.
+	 * @param CustomSerializer The custom serializer.
 	 *
 	 * @return This instance (for method chaining).
 	 */
@@ -67,7 +67,7 @@ public:
 	/**
 	 * Configures the bridge to use a specific message transport technology.
 	 *
-	 * @param InTransport - The transport technology to use.
+	 * @param InTransport The transport technology to use.
 	 *
 	 * @return This instance (for method chaining).
 	 */
@@ -95,7 +95,7 @@ public:
 	 *
 	 * If no address is specified, one will be generated automatically.
 	 *
-	 * @param InAddress - The address to set.
+	 * @param InAddress The address to set.
 	 *
 	 * @return This instance (for method chaining).
 	 */
@@ -111,7 +111,7 @@ public:
 	/**
 	 * Builds the message bridge as configured.
 	 *
-	 * @return A new message bridge, or NULL if it couldn't be built.
+	 * @return A new message bridge, or nullptr if it couldn't be built.
 	 */
 	IMessageBridgePtr Build( )
 	{
@@ -144,7 +144,7 @@ public:
 	/**
 	 * Implicit conversion operator to build the message bridge as configured.
 	 *
-	 * @return A new message bridge, or NULL if it couldn't be built.
+	 * @return A new message bridge, or nullptr if it couldn't be built.
 	 */
 	operator IMessageBridgePtr( )
 	{
diff --git a/Engine/Source/Runtime/Messaging/Public/Common/MessageEndpoint.h b/Engine/Source/Runtime/Messaging/Public/Common/MessageEndpoint.h
--- a/Engine/Source/Runtime/Messaging/Public/Common/MessageEndpoint.h
+++ b/Engine/Source/Runtime/Messaging/Public/Common/MessageEndpoint.h
@@ -72,9 +72,9 @@ public:
 	/**
 	 * Creates and initializes a new instance.
 	 *
-	 * @param InName - The endpoint's name (for debugging purposes).
-	 * @param InBus - The message bus to attach this endpoint to.
-	 * @param InHandlers - The collection of message handlers to register.
+	 * @param InName The endpoint's name (for debugging purposes).
+	 * @param InBus The message bus to attach this endpoint to.
+	 * @param InHandlers The collection of message handlers to register.
 	 */
 	FMessageEndpoint( const FName& InName, const IMessageBusRef& InBus, const TArray<IMessageHandlerPtr>& InHandlers )
 		: Address(FGuid::NewGuid())
@@ -167,7 +167,7 @@ public:
 	 * code includes time consuming operations, because it will block the message router,
 	 * causing no other messages to be delivered in the meantime.
 	 *
-	 * @param NamedThread - The name of the thread to receive messages on.
+	 * @param NamedThread The name of the thread to receive messages on.
 	 */
 	void SetRecipientThread( ENamedThreads::Type NamedThread )
 	{
@@ -182,8 +182,8 @@ public:
 	 * The message is effectively delivered again to this endpoint after the
 	 * original sent time plus the time delay have elapsed.
 	 *
-	 * @param Message - The context of the message to defer.
-	 * @param Delay - The time delay.
+	 * @param Message The context of the message to defer.
+	 * @param Delay The time delay.
 	 */
 	void Defer( const IMessageContextRef& Context, const FTimespan& Delay )
 	{
@@ -200,10 +200,10 @@ public:
 	 *
 	 * Messages can only be forwarded to endpoints within the same process.
 	 *
-	 * @param Context - The context of the message to forward.
-	 * @param Recipients - The list of message recipients to forward the message to.
-	 * @param ForwardingScope - The scope of the forwarded message.
-	 * @param Delay - The time delay.
+	 * @param Context The context of the message to forward.
+	 * @param Recipients The list of message recipients to forward the message to.
+	 * @param ForwardingScope The scope of the forwarded message.
+	 * @param Delay The time delay.
 	 */
 	void Forward( const IMessageContextRef& Context, const TArray<FMessageAddress>& Recipients, EMessageScope::Type ForwardingScope, const FTimespan& Delay )
 	{
@@ -218,12 +218,12 @@ public:
 	/**
 	 * Publishes a message to all subscribed recipients within the specified scope.
 	 *
-	 * @param Message - The message to publish.
-	 * @param TypeInfo - The message's type information.
-	 * @param Scope - The message scope.
-	 * @param Fields - The message content.
-	 * @param Delay - The delay after which to publish the message.
-	 * @param Expiration - The time at which the message expires.
+	 * @param Message The message to publish.
+	 * @param TypeInfo The message's type information.
+	 * @param Scope The message scope.
+	 * @param Fields The message content.
+	 * @param Delay The delay after which to publish the message.
+	 * @param Expiration The time at which the message expires.
 	 */
 	void Publish( void* Message, UScriptStruct* TypeInfo, EMessageScope::Type Scope, const FTimespan& Delay, const FDateTime& Expiration )
 	{
@@ -238,12 +238,12 @@ public:
 	/**
 	 * Sends a message to the specified list of recipients.
 	 *
-	 * @param Message - The message to send.
-	 * @param TypeInfo - The message's type information.
-	 * @param Attachment - An optional binary data attachment.
-	 * @param Recipients - The message recipients.
-	 * @param Delay - The delay after which to send the message.
-	 * @param Expiration - The time at which the message expires.
+	 * @param Message The message to send.
+	 * @param TypeInfo The message's type information.
+	 * @param Attachment An optional binary data attachment.
+	 * @param Recipients The message recipients.
+	 * @param Delay The delay after which to send the message.
+	 * @param Expiration The time at which the message expires.
 	 */
 	void Send( void* Message, UScriptStruct* TypeInfo, const IMessageAttachmentPtr& Attachment, const TArray<FMessageAddress>& Recipients, const FTimespan& Delay, const FDateTime& Expiration )
 	{
@@ -258,8 +258,8 @@ public:
 	/**
 	 * Subscribes a message handler.
 	 *
-	 * @param MessageType - The type name of the messages to subscribe to.
-	 * @param ScopeRange - The range of message scopes to include in the subscription.
+	 * @param MessageType The type name of the messages to subscribe to.
+	 * @param ScopeRange The range of message scopes to include in the subscription.
 	 */
 	void Subscribe( const FName& MessageType, const FMessageScopeRange& ScopeRange )
 	{
@@ -274,7 +274,7 @@ public:
 	/**
 	 * Unsubscribes this endpoint from the specified message type.
 	 *
-	 * @param MessageType - The type of message to unsubscribe (NAME_All = all types).
+	 * @param MessageType The type of message to unsubscribe (NAME_All = all types).
 	 *
 	 * @see Subscribe
 	 */
@@ -381,7 +381,7 @@ public:
 	 * been enabled and no matching message handler handled it. The inbox is disabled by default and
 	 * must be enabled using the EnableInbox() method.
 	 *
-	 * @param OutContext - Will hold the context of the received message.
+	 * @param OutContext Will hold the context of the received message.
 	 *
 	 * @return true if a message was received, false if the inbox was empty.
 	 *
@@ -481,9 +481,9 @@ public:
 	 *
 	 * Messages can only be forwarded to endpoints within the same process.
 	 *
-	 * @param Context - The context of the message to forward.
-	 * @param Recipient - The address of the recipient to forward the message to.
-	 * @param ForwardingScope - The scope of the forwarded message.
+	 * @param Context The context of the message to forward.
+	 * @param Recipient The address of the recipient to forward the message to.
+	 * @param ForwardingScope The scope of the forwarded message.
 	 */
 	void Forward( const IMessageContextRef& Context, const FMessageAddress& Recipient, EMessageScope::Type ForwardingScope )
 	{
@@ -495,10 +495,10 @@ public:
 	 *
 	 * Messages can only be forwarded to endpoints within the same process.
 	 *
-	 * @param Context - The context of the message to forward.
-	 * @param Recipient - The address of the recipient to forward the message to.
-	 * @param ForwardingScope - The scope of the forwarded message.
-	 * @param Delay - The delay after which to publish the message.
+	 * @param Context The context of the message to forward.
+	 * @param Recipient The address of the recipient to forward the message to.
+	 * @param ForwardingScope The scope of the forwarded message.
+	 * @param Delay The delay after which to publish the message.
 	 */
 	void Forward( const IMessageContextRef& Context, const FMessageAddress& Recipient, EMessageScope::Type ForwardingScope, const FTimespan& Delay )
 	{
@@ -510,9 +510,9 @@ public:
 	 *
 	 * Messages can only be forwarded to endpoints within the same process.
 	 *
-	 * @param Context - The context of the message to forward.
-	 * @param Recipients - The list of message recipients to forward the message to.
-	 * @param ForwardingScope - The scope of the forwarded message.
+	 * @param Context The context of the message to forward.
+	 * @param Recipients The list of message recipients to forward the message to.
+	 * @param ForwardingScope The scope of the forwarded message.
 	 */
 	void Forward( const IMessageContextRef& Context, const TArray<FMessageAddress>& Recipients, EMessageScope::Type ForwardingScope )
 	{
@@ -522,7 +522,7 @@ public:
 	/**
 	 * Immediately publishes a message to all subscribed recipients.
 	 *
-	 * @param Message - The message to publish.
+	 * @param Message The message to publish.
 	 */
 	template<typename MessageType>
 	void Publish( MessageType* Message )
@@ -533,8 +533,8 @@ public:
 	/**
 	 * Immediately pa message to all subscribed recipients within the specified scope.
 	 *
-	 * @param Message - The message to publish.
-	 * @param Scope - The message scope.
+	 * @param Message The message to publish.
+	 * @param Scope The message scope.
 	 */
 	template<typename MessageType>
 	void Publish( MessageType* Message, EMessageScope::Type Scope )
@@ -545,8 +545,8 @@ public:
 	/**
 	 * Publishes a message to all subscribed recipients after a given delay.
 	 *
-	 * @param Message - The message to publish.
-	 * @param Delay - The delay after which to publish the message.
+	 * @param Message The message to publish.
+	 * @param Delay The delay after which to publish the message.
 	 */
 	template<typename MessageType>
 	void Publish( MessageType* Message, const FTimespan& Delay )
@@ -557,9 +557,9 @@ public:
 	/**
 	 * Publishes a message to all subscribed recipients within the specified scope after a given delay.
 	 *
-	 * @param Message - The message to publish.
-	 * @param Scope - The message scope.
-	 * @param Delay - The delay after which to publish the message.
+	 * @param Message The message to publish.
+	 * @param Scope The message scope.
+	 * @param Delay The delay after which to publish the message.
 	 */
 	template<typename MessageType>
 	void Publish( MessageType* Message, EMessageScope::Type Scope, const FTimespan& Delay )
@@ -570,11 +570,11 @@ public:
 	/**
 	 * Publishes a message to all subscribed recipients within the specified scope.
 	 *
-	 * @param Message - The message to publish.
-	 * @param Scope - The message scope.
-	 * @param Fields - The message content.
-	 * @param Delay - The delay after which to publish the message.
-	 * @param Expiration - The time at which the message expires.
+	 * @param Message The message to publish.
+	 * @param Scope The message scope.
+	 * @param Fields The message content.
+	 * @param Delay The delay after which to publish the message.
+	 * @param Expiration The time at which the message expires.
 	 */
 	template<typename MessageType>
 	void Publish( MessageType* Message, EMessageScope::Type Scope, const FTimespan& Delay, const FDateTime& Expiration )
@@ -585,9 +585,9 @@ public:
 	/**
 	 * Immediately sends a message to the specified recipient.
 	 *
-	 * @param MessageType - The type of message to send.
-	 * @param Message - The message to send.
-	 * @param Recipient - The message recipient.
+	 * @param MessageType The type of message to send.
+	 * @param Message The message to send.
+	 * @param Recipient The message recipient.
 	 */
 	template<typename MessageType>
 	void Send( MessageType* Message, const FMessageAddress& Recipient )
@@ -598,10 +598,10 @@ public:
 	/**
 	 * Sends a message to the specified recipient after a given delay.
 	 *
-	 * @param MessageType - The type of message to send.
-	 * @param Message - The message to send.
-	 * @param Recipient - The message recipient.
-	 * @param Delay - The delay after which to send the message.
+	 * @param MessageType The type of message to send.
+	 * @param Message The message to send.
+	 * @param Recipient The message recipient.
+	 * @param Delay The delay after which to send the message.
 	 */
 	template<typename MessageType>
 	void Send( MessageType* Message, const FMessageAddress& Recipient, const FTimespan& Delay )
@@ -612,11 +612,11 @@ public:
 	/**
 	 * Sends a message with fields and expiration to the specified recipient after a given delay.
 	 *
-	 * @param MessageType - The type of message to send.
-	 * @param Message - The message to send.
-	 * @param Recipient - The message recipient.
-	 * @param Expiration - The time at which the message expires.
-	 * @param Delay - The delay after which to send the message.
+	 * @param MessageType The type of message to send.
+	 * @param Message The message to send.
+	 * @param Recipient The message recipient.
+	 * @param Expiration The time at which the message expires.
+	 * @param Delay The delay after which to send the message.
 	 */
 	template<typename MessageType>
 	void Send( MessageType* Message, const FMessageAddress& Recipient, const FTimespan& Delay, const FDateTime& Expiration )
@@ -627,10 +627,10 @@ public:
 	/**
 	 * Sends a message with fields and attachment to the specified recipient.
 	 *
-	 * @param MessageType - The type of message to send.
-	 * @param Message - The message to send.
-	 * @param Attachment - An optional binary data attachment.
-	 * @param Recipient - The message recipient.
+	 * @param MessageType The type of message to send.
+	 * @param Message The message to send.
+	 * @param Attachment An optional binary data attachment.
+	 * @param Recipient The message recipient.
 	 */
 	template<typename MessageType>
 	void Send( MessageType* Message, const IMessageAttachmentPtr& Attachment, const FMessageAddress& Recipient )
@@ -641,12 +641,12 @@ public:
 	/**
 	 * Sends a message with fields, attachment and expiration to the specified recipient after a given delay.
 	 *
-	 * @param MessageType - The type of message to send.
-	 * @param Message - The message to send.
-	 * @param Attachment - An optional binary data attachment.
-	 * @param Recipient - The message recipient.
-	 * @param Expiration - The time at which the message expires.
-	 * @param Delay - The delay after which to send the message.
+	 * @param MessageType The type of message to send.
+	 * @param Message The message to send.
+	 * @param Attachment An optional binary data attachment.
+	 * @param Recipient The message recipient.
+	 * @param Expiration The time at which the message expires.
+	 * @param Delay The delay after which to send the message.
 	 */
 	template<typename MessageType>
 	void Send( MessageType* Message, const IMessageAttachmentPtr& Attachment, const FMessageAddress& Recipient, const FDateTime& Expiration, const FTimespan& Delay )
@@ -657,9 +657,9 @@ public:
 	/**
 	 * Immediately sends a message to the specified list of recipients.
 	 *
-	 * @param MessageType - The type of message to send.
-	 * @param Message - The message to send.
-	 * @param Recipients - The message recipients.
+	 * @param MessageType The type of message to send.
+	 * @param Message The message to send.
+	 * @param Recipients The message recipients.
 	 */
 	template<typename MessageType>
 	void Send( MessageType* Message, const TArray<FMessageAddress>& Recipients )
@@ -670,10 +670,10 @@ public:
 	/**
 	 * Sends a message to the specified list of recipients after a given delay after a given delay.
 	 *
-	 * @param MessageType - The type of message to send.
-	 * @param Message - The message to send.
-	 * @param Recipients - The message recipients.
-	 * @param Delay - The delay after which to send the message.
+	 * @param MessageType The type of message to send.
+	 * @param Message The message to send.
+	 * @param Recipients The message recipients.
+	 * @param Delay The delay after which to send the message.
 	 */
 	template<typename MessageType>
 	void Send( MessageType* Message, const TArray<FMessageAddress>& Recipients, const FTimespan& Delay )
@@ -684,11 +684,11 @@ public:
 	/**
 	 * Sends a message with fields and attachment to the specified list of recipients after a given delay.
 	 *
-	 * @param MessageType - The type of message to send.
-	 * @param Message - The message to send.
-	 * @param Attachment - An optional binary data attachment.
-	 * @param Recipients - The message recipients.
-	 * @param Delay - The delay after which to send the message.
+	 * @param MessageType The type of message to send.
+	 * @param Message The message to send.
+	 * @param Attachment An optional binary data attachment.
+	 * @param Recipients The message recipients.
+	 * @param Delay The delay after which to send the message.
 	 */
 	template<typename MessageType>
 	void Send( MessageType* Message, const IMessageAttachmentPtr& Attachment, const TArray<FMessageAddress>& Recipients, const FTimespan& Delay )
@@ -699,12 +699,12 @@ public:
 	/**
 	 * Sends a message to the specified list of recipients.
 	 *
-	 * @param MessageType - The type of message to send.
-	 * @param Message - The message to send.
-	 * @param Attachment - An optional binary data attachment.
-	 * @param Recipients - The message recipients.
-	 * @param Delay - The delay after which to send the message.
-	 * @param Expiration - The time at which the message expires.
+	 * @param MessageType The type of message to send.
+	 * @param Message The message to send.
+	 * @param Attachment An optional binary data attachment.
+	 * @param Recipients The message recipients.
+	 * @param Delay The delay after which to send the message.
+	 * @param Expiration The time at which the message expires.
 	 */
 	template<typename MessageType>
 	void Send( MessageType* Message, const IMessageAttachmentPtr& Attachment, const TArray<FMessageAddress>& Recipients, const FTimespan& Delay, const FDateTime& Expiration )
@@ -717,10 +717,10 @@ public:
 	 *
 	 * The default message scope is all messages excluding loopback messages.
 	 *
-	 * @param HandlerType - The type of the class handling the message.
-	 * @param MessageType - The type of messages to subscribe to.
-	 * @param Handler - The class handling the messages.
-	 * @param HandlerFunc - The class function handling the messages.
+	 * @param HandlerType The type of the class handling the message.
+	 * @param MessageType The type of messages to subscribe to.
+	 * @param Handler The class handling the messages.
+	 * @param HandlerFunc The class function handling the messages.
 	 */
 	template<class MessageType>
 	void Subscribe( )
@@ -731,11 +731,11 @@ public:
 	/**
 	 * Template method to subscribe the message endpoint to the specified type and scope of messages.
 	 *
-	 * @param HandlerType - The type of the class handling the message.
-	 * @param MessageType - The type of messages to subscribe to.
-	 * @param Handler - The class handling the messages.
-	 * @param HandlerFunc - The class function handling the messages.
-	 * @param ScopeRange - The range of message scopes to include in the subscription.
+	 * @param HandlerType The type of the class handling the message.
+	 * @param MessageType The type of messages to subscribe to.
+	 * @param Handler The class handling the messages.
+	 * @param HandlerFunc The class function handling the messages.
+	 * @param ScopeRange The range of message scopes to include in the subscription.
 	 */
 	template<class MessageType>
 	void Subscribe( const FMessageScopeRange& ScopeRange )
@@ -756,7 +756,7 @@ public:
 	/**
 	 * Template method to unsubscribe the endpoint from the specified message type.
 	 *
-	 * @param MessageType - The type of message to unsubscribe (NAME_All = all types).
+	 * @param MessageType The type of message to unsubscribe (NAME_All = all types).
 	 *
 	 * @see Subscribe
 	 */
@@ -780,7 +780,7 @@ public:
 	 * Note: When calling this function make sure that no other object is holding on to
 	 * the endpoint, or otherwise the caller may get blocked forever.
 	 *
-	 * @param Endpoint - The message endpoint to release.
+	 * @param Endpoint The message endpoint to release.
 	 */
 	static void SafeRelease( FMessageEndpointPtr& Endpoint )
 	{
@@ -809,7 +809,7 @@ protected:
 	/**
 	 * Forwards the given message context to matching message handlers.
 	 *
-	 * @param Context - The context of the message to handle.
+	 * @param Context The context of the message to handle.
 	 */
 	void ProcessMessage( const IMessageContextRef& Context )
 	{
diff --git a/Engine/Source/Runtime/Messaging/Public/Common/MessageEndpointBuilder.h b/Engine/Source/Runtime/Messaging/Public/Common/MessageEndpointBuilder.h
--- a/Engine/Source/Runtime/Messaging/Public/Common/MessageEndpointBuilder.h
+++ b/Engine/Source/Runtime/Messaging/Public/Common/MessageEndpointBuilder.h
@@ -19,8 +19,8 @@ public:
 	 *
 	 * WARNING: This constructor must be called from the main thread.
 	 *
-	 * @param InName - The endpoint's name (for debugging purposes).
-	 * @param InBus - The message bus to attach the endpoint to.
+	 * @param InName The endpoint's name (for debugging purposes).
+	 * @param InBus The message bus to attach the endpoint to.
 	 */
 	FMessageEndpointBuilder( const FName& InName )
 		: BusPtr(IMessagingModule::Get().GetDefaultBus())
@@ -33,8 +33,8 @@ public:
 	/**
 	 * Creates and initializes a new builder using the specified message bus.
 	 *
-	 * @param InName - The endpoint's name (for debugging purposes).
-	 * @param InBus - The message bus to attach the endpoint to.
+	 * @param InName The endpoint's name (for debugging purposes).
+	 * @param InBus The message bus to attach the endpoint to.
 	 */
 	FMessageEndpointBuilder( const FName& InName, const IMessageBusRef& InBus )
 		: BusPtr(InBus)
@@ -49,11 +49,11 @@ public:
 	/**
 	 * Adds a message handler for the given type of messages.
 	 *
-	 * @param HandlerType - The type of the object handling the messages.
-	 * @param MessageType - The type of messages to handle.
-	 * @param Handler - The class handling the messages.
-	 * @param HandlerFunc - The class function handling the messages.
-	 * @param ReceivingThread - The thread on which to handle the message.
+	 * @param HandlerType The type of the object handling the messages.
+	 * @param MessageType The type of messages to handle.
+	 * @param Handler The class handling the messages.
+	 * @param HandlerFunc The class function handling the messages.
+	 * @param ReceivingThread The thread on which to handle the message.
 	 *
 	 * @return This instance (for method chaining).
 	 */
@@ -99,7 +99,7 @@ public:
 	 *
 	 * Also see the additional notes for ReceivingOnAnyThread().
 	 *
-	 * @param NamedThread - The name of the thread to receive messages on.
+	 * @param NamedThread The name of the thread to receive messages on.
 	 *
 	 * @return This instance (for method chaining).
 	 *
@@ -127,7 +127,7 @@ public:
 	/**
 	 * Registers a message handler with the endpoint.
 	 *
-	 * @param Handler - The handler to add.
+	 * @param Handler The handler to add.
 	 *
 	 * @return This instance (for method chaining).
 	 */
diff --git a/Engine/Source/Runtime/Messaging/Public/Interfaces/IAuthorizeMessageRecipients.h b/Engine/Source/Runtime/Messaging/Public/Interfaces/IAuthorizeMessageRecipients.h
--- a/Engine/Source/Runtime/Messaging/Public/Interfaces/IAuthorizeMessageRecipients.h
+++ b/Engine/Source/Runtime/Messaging/Public/Interfaces/IAuthorizeMessageRecipients.h
@@ -28,8 +28,8 @@ public:
 	/**
 	 * Authorizes a request to intercept messages of the specified type.
 	 *
-	 * @param Interceptor - The message interceptor to authorize.
-	 * @param MessageType - The type of messages to intercept.
+	 * @param Interceptor The message interceptor to authorize.
+	 * @param MessageType The type of messages to intercept.
 	 *
 	 * @return true if the request was authorized, false otherwise.
 	 */
@@ -38,8 +38,8 @@ public:
 	/**
 	 * Authorizes a request to register the specified recipient.
 	 *
-	 * @param Recipient - The recipient to register.
-	 * @param Address - The recipient's address.
+	 * @param Recipient The recipient to register.
+	 * @param Address The recipient's address.
 	 *
 	 * @return true if the request was authorized, false otherwise.
 	 */
@@ -48,8 +48,8 @@ public:
 	/**
 	 * Authorizes a request to add a subscription for the specified topic pattern.
 	 *
-	 * @param Subscriber - The subscriber.
-	 * @param TopicPattern - The message topic pattern to subscribe to.
+	 * @param Subscriber The subscriber.
+	 * @param TopicPattern The message topic pattern to subscribe to.
 	 *
 	 * @return true if the request is authorized, false otherwise.
 	 */
@@ -58,7 +58,7 @@ public:
 	/**
 	 * Authorizes a request to unregister the specified recipient.
 	 *
-	 * @param Address - The address of the recipient to unregister.
+	 * @param Address The address of the recipient to unregister.
 	 *
 	 * @return true if the request was authorized, false otherwise.
 	 */
@@ -67,8 +67,8 @@ public:
 	/**
 	 * Authorizes a request to remove a subscription for the specified topic pattern.
 	 *
-	 * @param Subscriber - The subscriber.
-	 * @param TopicPattern - The message topic pattern to unsubscribe from.
+	 * @param Subscriber The subscriber.
+	 * @param TopicPattern The message topic pattern to unsubscribe from.
 	 *
 	 * @return true if the request is authorized, false otherwise.
 	 */
diff --git a/Engine/Source/Runtime/Messaging/Public/Interfaces/IInterceptMessages.h b/Engine/Source/Runtime/Messaging/Public/Interfaces/IInterceptMessages.h
--- a/Engine/Source/Runtime/Messaging/Public/Interfaces/IInterceptMessages.h
+++ b/Engine/Source/Runtime/Messaging/Public/Interfaces/IInterceptMessages.h
@@ -28,7 +28,7 @@ public:
 	/**
 	 * Intercepts a message before it is being passed to the message router.
 	 *
-	 * @param Context - The context of the message to intercept.
+	 * @param Context The context of the message to intercept.
 	 *
 	 * @return true if the message was intercepted and should not be routed, false otherwise.
 	 */
diff --git a/Engine/Source/Runtime/Messaging/Public/Interfaces/IMessageBus.h b/Engine/Source/Runtime/Messaging/Public/Interfaces/IMessageBus.h
--- a/Engine/Source/Runtime/Messaging/Public/Interfaces/IMessageBus.h
+++ b/Engine/Source/Runtime/Messaging/Public/Interfaces/IMessageBus.h
@@ -41,11 +41,11 @@ public:
 	 *
 	 * Messages can only be forwarded to endpoints within the same process.
 	 *
-	 * @param Context - The context of the message to forward.
-	 * @param Recipients - The list of message recipients to forward the message to.
-	 * @param NewScope - The forwarding scope.
-	 * @param Delay - The deferral time.
-	 * @param Forwarder - The sender that forwards the message.
+	 * @param Context The context of the message to forward.
+	 * @param Recipients The list of message recipients to forward the message to.
+	 * @param NewScope The forwarding scope.
+	 * @param Delay The deferral time.
+	 * @param Forwarder The sender that forwards the message.
 	 */
 	virtual void Forward( const IMessageContextRef& Context, const TArray<FMessageAddress>& Recipients, EMessageScope::Type NewScope, const FTimespan& Delay, const ISendMessagesRef& Forwarder ) = 0;
 
@@ -59,8 +59,8 @@ public:
 	/**
 	 * Adds an interceptor for messages of the specified type.
 	 *
-	 * @param Interceptor - The interceptor.
-	 * @param MessageType - The type of messages to intercept.
+	 * @param Interceptor The interceptor.
+	 * @param MessageType The type of messages to intercept.
 	 */
 	virtual void Intercept( const IInterceptMessagesRef& Interceptor, const FName& MessageType ) = 0;
 
@@ -69,33 +69,33 @@ public:
 	 *
 	 * The message is published with a default time to live of Process.
 	 *
-	 * @param Message - The message to publish.
-	 * @param TypeInfo - The message's type information.
-	 * @param Scope - The message scope.
-	 * @param Delay - The delay after which to send the message.
-	 * @param Expiration - The time at which the message expires.
-	 * @param Publisher - The message publisher.
+	 * @param Message The message to publish.
+	 * @param TypeInfo The message's type information.
+	 * @param Scope The message scope.
+	 * @param Delay The delay after which to send the message.
+	 * @param Expiration The time at which the message expires.
+	 * @param Publisher The message publisher.
 	 */
 	virtual void Publish( void* Message, UScriptStruct* TypeInfo, EMessageScope::Type Scope, const FTimespan& Delay, const FDateTime& Expiration, const ISendMessagesRef& Publisher ) = 0;
 
 	/**
 	 * Registers a message recipient with the message bus.
 	 *
-	 * @param Address - The address of the recipient to register.
-	 * @param Recipient - The message recipient.
+	 * @param Address The address of the recipient to register.
+	 * @param Recipient The message recipient.
 	 */
 	virtual void Register( const FMessageAddress& Address, const IReceiveMessagesRef& Recipient) = 0;
 
 	/**
 	 * Sends a message to multiple recipients.
 	 *
-	 * @param Message - The message to send.
-	 * @param TypeInfo - The message's type information.
-	 * @param Attachment - The binary data to attach to the message.
-	 * @param Recipients - The list of message recipients.
-	 * @param Delay - The delay after which to send the message.
-	 * @param Expiration - The time at which the message expires.
-	 * @param Sender - The message sender.
+	 * @param Message The message to send.
+	 * @param TypeInfo The message's type information.
+	 * @param Attachment The binary data to attach to the message.
+	 * @param Recipients The list of message recipients.
+	 * @param Delay The delay after which to send the message.
+	 * @param Expiration The time at which the message expires.
+	 * @param Sender The message sender.
 	 */
 	virtual void Send( void* Message, UScriptStruct* TypeInfo, const IMessageAttachmentPtr& Attachment, const TArray<FMessageAddress>& Recipients, const FTimespan& Delay, const FDateTime& Expiration, const ISendMessagesRef& Sender ) = 0;
 
@@ -112,9 +112,9 @@ public:
 	 * Subscriptions allow message consumers to receive published messages from the message bus.
 	 * The returned interface can be used to query the subscription's details and its enabled state.
 	 *
-	 * @param Subscriber - The subscriber wishing to receive the messages.
-	 * @param MessageType - The type of messages to subscribe to (NAME_All = subscribe to all message types).
-	 * @param ScopeRange - The range of message scopes to include in the subscription.
+	 * @param Subscriber The subscriber wishing to receive the messages.
+	 * @param MessageType The type of messages to subscribe to (NAME_All = subscribe to all message types).
+	 * @param ScopeRange The range of message scopes to include in the subscription.
 	 *
 	 * @return The added subscription, or NULL if the subscription failed.
 	 *
@@ -125,23 +125,23 @@ public:
 	/**
 	 * Removes an interceptor for messages of the specified type.
 	 *
-	 * @param Interceptor - The interceptor to remove.
-	 * @param MessageType - The type of messages to stop intercepting.
+	 * @param Interceptor The interceptor to remove.
+	 * @param MessageType The type of messages to stop intercepting.
 	 */
 	virtual void Unintercept( const IInterceptMessagesRef& Interceptor, const FName& MessageType ) = 0;
 
 	/**
 	 * Unregisters a message recipient from the message bus.
 	 *
-	 * @param Address - The address of the recipient to unregister.
+	 * @param Address The address of the recipient to unregister.
 	 */
 	virtual void Unregister( const FMessageAddress& Address ) = 0;
 
 	/**
 	 * Cancels the specified message subscription.
 	 *
-	 * @param Subscriber - The subscriber wishing to stop receiving the messages.
-	 * @param MessageType - The type of messages to unsubscribe from (NAME_All = all types).
+	 * @param Subscriber The subscriber wishing to stop receiving the messages.
+	 * @param MessageType The type of messages to unsubscribe from (NAME_All = all types).
 	 *
 	 * @see Subscribe
 	 */
diff --git a/Engine/Source/Runtime/Messaging/Public/Interfaces/IMessageHandler.h b/Engine/Source/Runtime/Messaging/Public/Interfaces/IMessageHandler.h
--- a/Engine/Source/Runtime/Messaging/Public/Interfaces/IMessageHandler.h
+++ b/Engine/Source/Runtime/Messaging/Public/Interfaces/IMessageHandler.h
@@ -35,7 +35,7 @@ public:
 	/**
 	 * Handles the specified message.
 	 *
-	 * @param Context - The context of the message to handle.
+	 * @param Context The context of the message to handle.
 	 */
 	virtual void HandleMessage( const IMessageContextRef& Context ) = 0;
 
@@ -51,8 +51,8 @@ public:
 /**
  * Template for message handler function pointer types.
  *
- * @param MessageType - The type of message to handle.
- * @param HandlerType - The type of the handler class.
+ * @param MessageType The type of message to handle.
+ * @param HandlerType The type of the handler class.
  */
 template<typename MessageType, typename HandlerType>
 struct TMessageHandlerFunc
@@ -64,8 +64,8 @@ struct TMessageHandlerFunc
 /**
  * Template for message handlers.
  *
- * @param MessageType - The type of message to handle.
- * @param HandlerType - The type of the handler class.
+ * @param MessageType The type of message to handle.
+ * @param HandlerType The type of the handler class.
  */
 template<typename MessageType, typename HandlerType>
 class TMessageHandler
@@ -76,8 +76,8 @@ public:
 	/**
 	 * Creates and initializes a new message handler.
 	 *
-	 * @param InHandler - The object that will handle the messages.
-	 * @param InHandlerFunc - The object's message handling function.
+	 * @param InHandler The object that will handle the messages.
+	 * @param InHandlerFunc The object's message handling function.
 	 */
 	TMessageHandler( HandlerType* InHandler, typename TMessageHandlerFunc<MessageType, HandlerType>::Type InHandlerFunc )
 		: Handler(InHandler)
diff --git a/Engine/Source/Runtime/Messaging/Public/Interfaces/IMessageTracer.h b/Engine/Source/Runtime/Messaging/Public/Interfaces/IMessageTracer.h
--- a/Engine/Source/Runtime/Messaging/Public/Interfaces/IMessageTracer.h
+++ b/Engine/Source/Runtime/Messaging/Public/Interfaces/IMessageTracer.h
@@ -317,7 +317,7 @@ public:
 	/**
 	 * Ticks the tracer.
 	 *
-	 * @param DeltaTime - The time in seconds since the last tick.
+	 * @param DeltaTime The time in seconds since the last tick.
 	 *
 	 * @return true if any events were processed.
 	 */
@@ -328,7 +328,7 @@ public:
 	/**
 	 * Gets the list of known message endpoints.
 	 *
-	 * @param OutEndpoints - Will contain the list of endpoints.
+	 * @param OutEndpoints Will contain the list of endpoints.
 	 *
 	 * @return The number of endpoints returned.
 	 */
@@ -344,8 +344,8 @@ public:
 	/**
 	 * Gets the list of known message types filtered by name.
 	 *
-	 * @param NameFilter - The name substring to filter with.
-	 * @param OutTypes - Will contain the list of message types.
+	 * @param NameFilter The name substring to filter with.
+	 * @param OutTypes Will contain the list of message types.
 	 *
 	 * @return The number of message types returned.
 	 */
diff --git a/Engine/Source/Runtime/Messaging/Public/Interfaces/IMessageTracerBreakpoint.h b/Engine/Source/Runtime/Messaging/Public/Interfaces/IMessageTracerBreakpoint.h
--- a/Engine/Source/Runtime/Messaging/Public/Interfaces/IMessageTracerBreakpoint.h
+++ b/Engine/Source/Runtime/Messaging/Public/Interfaces/IMessageTracerBreakpoint.h
@@ -35,7 +35,7 @@ public:
 	/**
 	 * Checks whether the tracer should break on the given message.
 	 *
-	 * @param Context - The message to break on.
+	 * @param Context The message to break on.
 	 *
 	 * @return true if the tracer should break, false otherwise.
 	 */
diff --git a/Engine/Source/Runtime/Messaging/Public/Interfaces/IMessagingModule.h b/Engine/Source/Runtime/Messaging/Public/Interfaces/IMessagingModule.h
--- a/Engine/Source/Runtime/Messaging/Public/Interfaces/IMessagingModule.h
+++ b/Engine/Source/Runtime/Messaging/Public/Interfaces/IMessagingModule.h
@@ -21,10 +21,10 @@ public:
 	 * Message bridges translate messages between a message bus and another means of
 	 * message transportation, such as network sockets.
 	 *
-	 * @param Address - The bridge's address on the message bus.
-	 * @param Bus - The message bus to attach the bridge to.
-	 * @param Serializer - The message serializer to use.
-	 * @param Transport - The message transport technology to use.
+	 * @param Address The bridge's address on the message bus.
+	 * @param Bus The message bus to attach the bridge to.
+	 * @param Serializer The message serializer to use.
+	 * @param Transport The message transport technology to use.
 	 *
 	 * @return The new message bridge, or NULL if the bridge couldn't be created.
 	 *
@@ -35,7 +35,7 @@ public:
 	/**
 	 * Creates a new message bus.
 	 *
-	 * @param RecipientAuthorizer - An optional recipient authorizer.
+	 * @param RecipientAuthorizer An optional recipient authorizer.
 	 *
 	 * @return The new message bus, or NULL if the bus couldn't be created.
 	 *
@@ -62,7 +62,7 @@ public:
 	/**
 	 * Gets a reference to the messaging module instance.
 	 *
-	 * @todo gmp - better implementation using dependency injection.
+	 * @todo gmp: better implementation using dependency injection.
 	 *
 	 * @return A reference to the Messaging module.
 	 */
diff --git a/Engine/Source/Runtime/Messaging/Public/Interfaces/IMutableMessageContext.h b/Engine/Source/Runtime/Messaging/Public/Interfaces/IMutableMessageContext.h
--- a/Engine/Source/Runtime/Messaging/Public/Interfaces/IMutableMessageContext.h
+++ b/Engine/Source/Runtime/Messaging/Public/Interfaces/IMutableMessageContext.h
@@ -29,58 +29,58 @@ public:
 	/**
 	 * Adds a message address to the recipient list.
 	 *
-	 * @param Recipient - The address of the recipient to add.
+	 * @param Recipient The address of the recipient to add.
 	 */
 	virtual void AddRecipient( const FMessageAddress& Recipient ) = 0;
 
 	/**
 	 * Sets the optional message attachment.
 	 *
-	 * @param InAttachment - The attachment to set.
+	 * @param InAttachment The attachment to set.
 	 */
 	virtual void SetAttachment( const IMessageAttachmentPtr& InAttachment ) = 0;
 
 	/**
 	 * Sets the message.
 	 *
-	 * @param InMessage - The message to set.
-	 * @param InTypeInfo - The message's type information.
+	 * @param InMessage The message to set.
+	 * @param InTypeInfo The message's type information.
 	 */
 	virtual void SetMessage( void* InMessage, UScriptStruct* InTypeInfo ) = 0;
 
 	/**
 	 * Sets the date and time at which the message expires.
 	 *
-	 * @param InExpiration - Expiration date and time.
+	 * @param InExpiration Expiration date and time.
 	 */
 	virtual void SetExpiration( const FDateTime& InExpiration ) = 0;
 
 	/**
 	 * Sets the value of the header with the specified key.
 	 *
-	 * @param Key - The header key.
-	 * @param Value - The header value.
+	 * @param Key The header key.
+	 * @param Value The header value.
 	 */
 	virtual void SetHeader( const FName& Key, const FString& Value ) = 0;
 
 	/**
 	 * Sets the message scope.
 	 *
-	 * @param InScope - The message scope.
+	 * @param InScope The message scope.
 	 */
 	virtual void SetScope( EMessageScope::Type InScope ) = 0;
 
 	/**
 	 * Sets the address of the message's sender.
 	 *
-	 * @param InSender - The message sender's address.
+	 * @param InSender The message sender's address.
 	 */
 	virtual void SetSender( const FMessageAddress& InSender ) = 0;
 
 	/**
 	 * Sets the date and time at which the message was sent.
 	 *
-	 * @param InTimeSent - Send date and time.
+	 * @param InTimeSent Send date and time.
 	 */
 	virtual void SetTimeSent( const FDateTime& InTimeSent ) = 0;
 
diff --git a/Engine/Source/Runtime/Messaging/Public/Interfaces/IReceiveMessages.h b/Engine/Source/Runtime/Messaging/Public/Interfaces/IReceiveMessages.h
--- a/Engine/Source/Runtime/Messaging/Public/Interfaces/IReceiveMessages.h
+++ b/Engine/Source/Runtime/Messaging/Public/Interfaces/IReceiveMessages.h
@@ -75,7 +75,7 @@ public:
 	/**
 	 * Handles the given message.
 	 *
-	 * @param Context - Will hold the context of the received message.
+	 * @param Context Will hold the context of the received message.
 	 */
 	virtual void ReceiveMessage( const IMessageContextRef& Context ) = 0;
 
diff --git a/Engine/Source/Runtime/Messaging/Public/Interfaces/ISendMessages.h b/Engine/Source/Runtime/Messaging/Public/Interfaces/ISendMessages.h
--- a/Engine/Source/Runtime/Messaging/Public/Interfaces/ISendMessages.h
+++ b/Engine/Source/Runtime/Messaging/Public/Interfaces/ISendMessages.h
@@ -32,8 +32,8 @@ public:
 	/**
 	 * Notifies the sender of errors.
 	 *
-	 * @param Context - The context of the message that generated the error.
-	 * @param Error - The error string.
+	 * @param Context The context of the message that generated the error.
+	 * @param Error The error string.
 	 */
 	virtual void NotifyMessageError( const IMessageContextRef& Context, const FString& Error ) = 0;
 
diff --git a/Engine/Source/Runtime/Messaging/Public/Interfaces/ISerializeMessages.h b/Engine/Source/Runtime/Messaging/Public/Interfaces/ISerializeMessages.h
--- a/Engine/Source/Runtime/Messaging/Public/Interfaces/ISerializeMessages.h
+++ b/Engine/Source/Runtime/Messaging/Public/Interfaces/ISerializeMessages.h
@@ -31,8 +31,8 @@ public:
 	/**
 	 * Deserializes a message from an archive.
 	 *
-	 * @param Data - The archive to serialize from.
-	 * @param OutContext - Will hold the context of the deserialized message.
+	 * @param Data The archive to serialize from.
+	 * @param OutContext Will hold the context of the deserialized message.
 	 *
 	 * @return true if deserialization was successful, false otherwise.
 	 */
@@ -41,8 +41,8 @@ public:
 	/**
 	 * Serializes a message into an archive.
 	 *
-	 * @param Context - The context of the message to serialize.
-	 * @param Archive - The archive to serialize into.
+	 * @param Context The context of the message to serialize.
+	 * @param Archive The archive to serialize into.
 	 *
 	 * @return true if serialization was successful, false otherwise.
 	 */
diff --git a/Engine/Source/Runtime/Messaging/Public/Interfaces/ITransportMessages.h b/Engine/Source/Runtime/Messaging/Public/Interfaces/ITransportMessages.h
--- a/Engine/Source/Runtime/Messaging/Public/Interfaces/ITransportMessages.h
+++ b/Engine/Source/Runtime/Messaging/Public/Interfaces/ITransportMessages.h
@@ -81,9 +81,9 @@ public:
 	/**
 	 * Transports the given message data to the specified network nodes.
 	 *
-	 * @param Data - The serialized message data to transport.
-	 * @param Attachment - An optional message attachment (i.e. file or memory buffer).
-	 * @param Recipients - The transport nodes to send the message to.
+	 * @param Data The serialized message data to transport.
+	 * @param Attachment An optional message attachment (i.e. file or memory buffer).
+	 * @param Recipients The transport nodes to send the message to.
 	 *
 	 * @return true if the message is being transported, false otherwise.
 	 */
@@ -94,21 +94,21 @@ public:
 	/**
 	 * Returns a delegate that is executed when message data has been received.
 	 *
-	 * @param Delegate - The delegate to set.
+	 * @param Delegate The delegate to set.
 	 */
 	virtual FOnMessageTransportMessageReceived& OnMessageReceived( ) = 0;
 
 	/**
 	 * Returns a delegate that is executed when a transport node has been discovered.
 	 *
-	 * @param Delegate - The delegate to set.
+	 * @param Delegate The delegate to set.
 	 */
 	virtual FOnMessageTransportNodeDiscovered& OnNodeDiscovered( ) = 0;
 
 	/**
 	 * Returns a delegate that is executed when a transport node has closed or timed out.
 	 *
-	 * @param Delegate - The delegate to set.
+	 * @param Delegate The delegate to set.
 	 */
 	virtual FOnMessageTransportNodeLost& OnNodeLost( ) = 0;
 
diff --git a/Engine/Source/Runtime/Messaging/Public/Messaging.h b/Engine/Source/Runtime/Messaging/Public/Messaging.h
--- a/Engine/Source/Runtime/Messaging/Public/Messaging.h
+++ b/Engine/Source/Runtime/Messaging/Public/Messaging.h
@@ -7,7 +7,7 @@
 #pragma once
 
 
-/* Dependencies
+/* Public Dependencies
  *****************************************************************************/
 
 #include "CoreUObject.h"
@@ -15,7 +15,7 @@
 #include "TaskGraphInterfaces.h"
 
 
-/* Interfaces
+/* Public Interfaces
  *****************************************************************************/
 
 #include "IMessageAttachment.h"
@@ -37,7 +37,7 @@
 #include "IMessagingModule.h"
 
 
-/* Common
+/* Public Common
  *****************************************************************************/
 
 #include "FileMessageAttachment.h"
