Messaging: Documentation cleanup pass

diff --git a/Engine/Source/Runtime/Messaging/Private/Bridge/MessageAddressBook.h b/Engine/Source/Runtime/Messaging/Private/Bridge/MessageAddressBook.h
--- a/Engine/Source/Runtime/Messaging/Private/Bridge/MessageAddressBook.h
+++ b/Engine/Source/Runtime/Messaging/Private/Bridge/MessageAddressBook.h
@@ -59,7 +59,6 @@ public:
 	 * Checks whether this address book contains the given address.
 	 *
 	 * @param Address The address to check.
-	 *
 	 * @return true if the address is known, false otherwise.
 	 */
 	bool Contains( const FMessageAddress& Address )
@@ -73,7 +72,6 @@ public:
 	 * Gets the transport node identifiers for the specified list of message addresses.
 	 *
 	 * @param Addresses The address list to retrieve the node identifiers for.
-	 *
 	 * @return The list of node identifiers.
 	 */
 	TArray<FGuid> GetNodesFor( const TArray<FMessageAddress>& Addresses )
@@ -102,9 +100,7 @@ public:
 	 * If you are not interested in the removed addresses, use Clear().
 	 *
 	 * @param OutRemovedRecipients Will hold a list of recipients that were removed.
-	 *
-	 * @see Clear
-	 * @see RemoveNode
+	 * @see Clear, RemoveNode
 	 */
 	void RemoveAllNodes( TArray<FMessageAddress>& OutRemovedAddresses )
 	{
@@ -121,9 +117,7 @@ public:
 	 *
 	 * @param NodeId The identifier of the transport node to remove.
 	 * @param OutRemovedRecipients Will hold a list of recipients that were removed.
-	 *
-	 * @see Clear
-	 * @see RemoveAllNodes
+	 * @see Clear, RemoveAllNodes
 	 */
 	void RemoveNode( const FGuid& NodeId, TArray<FMessageAddress>& OutRemovedAddresses )
 	{
diff --git a/Engine/Source/Runtime/Messaging/Private/Bridge/MessageData.h b/Engine/Source/Runtime/Messaging/Private/Bridge/MessageData.h
--- a/Engine/Source/Runtime/Messaging/Private/Bridge/MessageData.h
+++ b/Engine/Source/Runtime/Messaging/Private/Bridge/MessageData.h
@@ -7,14 +7,10 @@
 #pragma once
 
 
-/**
- * Type definition for shared pointers to instances of FMessageData.
- */
+/** Type definition for shared pointers to instances of FMessageData. */
 typedef TSharedPtr<class FMessageData, ESPMode::ThreadSafe> FMessageDataPtr;
 
-/**
- * Type definition for shared references to instances of FMessageData.
- */
+/** Type definition for shared references to instances of FMessageData. */
 typedef TSharedRef<class FMessageData, ESPMode::ThreadSafe> FMessageDataRef;
 
 
diff --git a/Engine/Source/Runtime/Messaging/Private/Bridge/MessageDeserializeTask.h b/Engine/Source/Runtime/Messaging/Private/Bridge/MessageDeserializeTask.h
--- a/Engine/Source/Runtime/Messaging/Private/Bridge/MessageDeserializeTask.h
+++ b/Engine/Source/Runtime/Messaging/Private/Bridge/MessageDeserializeTask.h
@@ -18,6 +18,7 @@ public:
 	 * Creates and initializes a new instance.
 	 *
 	 * InMessageTypeInfo The message's type information (may be nullptr if the type information hasn't been discovered yet).
+	 *
 	 * @param InSerializer The serializer to use.
 	 * @param InThread The name of the thread to deserialize the message on.
 	 */
diff --git a/Engine/Source/Runtime/Messaging/Private/Bus/MessageRouter.h b/Engine/Source/Runtime/Messaging/Private/Bus/MessageRouter.h
--- a/Engine/Source/Runtime/Messaging/Private/Bus/MessageRouter.h
+++ b/Engine/Source/Runtime/Messaging/Private/Bus/MessageRouter.h
@@ -139,7 +139,6 @@ protected:
 	 * Queues up a router command.
 	 *
 	 * @param Command The command to queue up.
-	 *
 	 * @return true if the command was enqueued, false otherwise.
 	 */
 	FORCEINLINE bool EnqueueCommand( TBaseDelegate_NoParams<void> Command )
@@ -237,13 +236,13 @@ private:
 private:
 
 	// Maps message types to interceptors.
-	TMap<FName, TArray<IInterceptMessagesPtr> > ActiveInterceptors;
+	TMap<FName, TArray<IInterceptMessagesPtr>> ActiveInterceptors;
 
 	// Maps message addresses to recipients.
 	TMap<FMessageAddress, IReceiveMessagesWeakPtr> ActiveRecipients;
 
 	// Maps message types to subscriptions.
-	TMap<FName, TArray<IMessageSubscriptionPtr> > ActiveSubscriptions;
+	TMap<FName, TArray<IMessageSubscriptionPtr>> ActiveSubscriptions;
 
 	// Holds the router command queue.
 	TQueue<TBaseDelegate_NoParams<void>, EQueueMode::Mpsc> Commands;
diff --git a/Engine/Source/Runtime/Messaging/Private/Bus/MessageTracer.h b/Engine/Source/Runtime/Messaging/Private/Bus/MessageTracer.h
--- a/Engine/Source/Runtime/Messaging/Private/Bus/MessageTracer.h
+++ b/Engine/Source/Runtime/Messaging/Private/Bus/MessageTracer.h
@@ -7,19 +7,13 @@
 #pragma once
 
 
-/**
- * Type definition for weak pointers to instances of FMessageTracer.
- */
+/** Type definition for weak pointers to instances of FMessageTracer. */
 typedef TWeakPtr<class FMessageTracer, ESPMode::ThreadSafe> FMessageTracerWeakPtr;
 
-/**
- * Type definition for shared pointers to instances of FMessageTracer.
- */
+/** Type definition for shared pointers to instances of FMessageTracer. */
 typedef TSharedPtr<class FMessageTracer, ESPMode::ThreadSafe> FMessageTracerPtr;
 
-/**
- * Type definition for shared references to instances of FMessageTracer.
- */
+/** Type definition for shared references to instances of FMessageTracer. */
 typedef TSharedRef<class FMessageTracer, ESPMode::ThreadSafe> FMessageTracerRef;
 
 
diff --git a/Engine/Source/Runtime/Messaging/Public/Common/MessageBridgeBuilder.h b/Engine/Source/Runtime/Messaging/Public/Common/MessageBridgeBuilder.h
--- a/Engine/Source/Runtime/Messaging/Public/Common/MessageBridgeBuilder.h
+++ b/Engine/Source/Runtime/Messaging/Public/Common/MessageBridgeBuilder.h
@@ -54,7 +54,6 @@ public:
 	 * Configures the bridge to use a custom message serializer.\
 	 *
 	 * @param CustomSerializer The custom serializer.
-	 *
 	 * @return This instance (for method chaining).
 	 */
 	FMessageBridgeBuilder& UsingCustomSerializer( const ISerializeMessagesRef& CustomSerializer )
@@ -68,7 +67,6 @@ public:
 	 * Configures the bridge to use a specific message transport technology.
 	 *
 	 * @param InTransport The transport technology to use.
-	 *
 	 * @return This instance (for method chaining).
 	 */
 	FMessageBridgeBuilder& UsingTransport( const ITransportMessagesRef& InTransport )
@@ -96,7 +94,6 @@ public:
 	 * If no address is specified, one will be generated automatically.
 	 *
 	 * @param InAddress The address to set.
-	 *
 	 * @return This instance (for method chaining).
 	 */
 	FMessageBridgeBuilder& WithAddress( const FMessageAddress& InAddress )
diff --git a/Engine/Source/Runtime/Messaging/Public/Common/MessageEndpoint.h b/Engine/Source/Runtime/Messaging/Public/Common/MessageEndpoint.h
--- a/Engine/Source/Runtime/Messaging/Public/Common/MessageEndpoint.h
+++ b/Engine/Source/Runtime/Messaging/Public/Common/MessageEndpoint.h
@@ -7,14 +7,10 @@
 #pragma once
 
 
-/**
- * Type definition for shared pointers to instances of FMessageEndpoint.
- */
+/** Type definition for shared pointers to instances of FMessageEndpoint. */
 typedef TSharedPtr<class FMessageEndpoint, ESPMode::ThreadSafe> FMessageEndpointPtr;
 
-/**
- * Type definition for shared references to instances of FMessageEndpoint.
- */
+/** Type definition for shared references to instances of FMessageEndpoint. */
 typedef TSharedRef<class FMessageEndpoint, ESPMode::ThreadSafe> FMessageEndpointRef;
 
 
@@ -275,7 +271,6 @@ public:
 	 * Unsubscribes this endpoint from the specified message type.
 	 *
 	 * @param MessageType The type of message to unsubscribe (NAME_All = all types).
-	 *
 	 * @see Subscribe
 	 */
 	void Unsubscribe( const FName& TopicPattern )
@@ -361,8 +356,7 @@ public:
 	 * been enabled and no matching message handler handled it. The inbox is disabled by default and
 	 * must be enabled using the EnableInbox() method.
 	 *
-	 * @see IsInboxEmpty
-	 * @see ReceiveFromInbox
+	 * @see IsInboxEmpty, ReceiveFromInbox
 	 */
 	void ProcessInbox( )
 	{
@@ -385,10 +379,7 @@ public:
 	 *
 	 * @return true if a message was received, false if the inbox was empty.
 	 *
-	 * @see DisableInbox
-	 * @see EnableInbox
-	 * @see IsInboxEnabled
-	 * @see ProcessInbox
+	 * @see DisableInbox, EnableInbox, IsInboxEnabled, ProcessInbox
 	 */
 	bool ReceiveFromInbox( IMessageContextPtr& OutContext )
 	{
@@ -757,7 +748,6 @@ public:
 	 * Template method to unsubscribe the endpoint from the specified message type.
 	 *
 	 * @param MessageType The type of message to unsubscribe (NAME_All = all types).
-	 *
 	 * @see Subscribe
 	 */
 	template<class MessageType>
diff --git a/Engine/Source/Runtime/Messaging/Public/Common/MessageEndpointBuilder.h b/Engine/Source/Runtime/Messaging/Public/Common/MessageEndpointBuilder.h
--- a/Engine/Source/Runtime/Messaging/Public/Common/MessageEndpointBuilder.h
+++ b/Engine/Source/Runtime/Messaging/Public/Common/MessageEndpointBuilder.h
@@ -54,7 +54,6 @@ public:
 	 * @param Handler The class handling the messages.
 	 * @param HandlerFunc The class function handling the messages.
 	 * @param ReceivingThread The thread on which to handle the message.
-	 *
 	 * @return This instance (for method chaining).
 	 */
 	template<typename MessageType, typename HandlerType>
@@ -81,8 +80,7 @@ public:
 	 * causing no other messages to be delivered in the meantime.
 	 *
 	 * @return This instance (for method chaining).
-	 *
-	 * @ReceivingOnThread
+	 * @see ReceivingOnThread
 	 */
 	FMessageEndpointBuilder& ReceivingOnAnyThread( )
 	{
@@ -100,10 +98,8 @@ public:
 	 * Also see the additional notes for ReceivingOnAnyThread().
 	 *
 	 * @param NamedThread The name of the thread to receive messages on.
-	 *
 	 * @return This instance (for method chaining).
-	 *
-	 * @seeReceivingOnAnyThread
+	 * @see ReceivingOnAnyThread
 	 */
 	FMessageEndpointBuilder& ReceivingOnThread( ENamedThreads::Type NamedThread )
 	{
@@ -128,7 +124,6 @@ public:
 	 * Registers a message handler with the endpoint.
 	 *
 	 * @param Handler The handler to add.
-	 *
 	 * @return This instance (for method chaining).
 	 */
 	FMessageEndpointBuilder& WithHandler( const IMessageHandlerRef& Handler )
diff --git a/Engine/Source/Runtime/Messaging/Public/Interfaces/IAuthorizeMessageRecipients.h b/Engine/Source/Runtime/Messaging/Public/Interfaces/IAuthorizeMessageRecipients.h
--- a/Engine/Source/Runtime/Messaging/Public/Interfaces/IAuthorizeMessageRecipients.h
+++ b/Engine/Source/Runtime/Messaging/Public/Interfaces/IAuthorizeMessageRecipients.h
@@ -7,14 +7,10 @@
 #pragma once
 
 
-/**
- * Type definition for shared pointers to instances of IAuthorizeMessageRecipients.
- */
+/** Type definition for shared pointers to instances of IAuthorizeMessageRecipients. */
 typedef TSharedPtr<class IAuthorizeMessageRecipients> IAuthorizeMessageRecipientsPtr;
 
-/**
- * Type definition for shared references to instances of IAuthorizeMessageRecipients.
- */
+/** Type definition for shared references to instances of IAuthorizeMessageRecipients. */
 typedef TSharedRef<class IAuthorizeMessageRecipients> IAuthorizeMessageRecipientsRef;
 
 
@@ -30,7 +26,6 @@ public:
 	 *
 	 * @param Interceptor The message interceptor to authorize.
 	 * @param MessageType The type of messages to intercept.
-	 *
 	 * @return true if the request was authorized, false otherwise.
 	 */
 	virtual bool AuthorizeInterceptor( const IInterceptMessagesRef& Interceptor, const FName& MessageType ) = 0;
@@ -40,7 +35,6 @@ public:
 	 *
 	 * @param Recipient The recipient to register.
 	 * @param Address The recipient's address.
-	 *
 	 * @return true if the request was authorized, false otherwise.
 	 */
 	virtual bool AuthorizeRegistration( const IReceiveMessagesRef& Recipient, const FMessageAddress& Address ) = 0;
@@ -50,7 +44,6 @@ public:
 	 *
 	 * @param Subscriber The subscriber.
 	 * @param TopicPattern The message topic pattern to subscribe to.
-	 *
 	 * @return true if the request is authorized, false otherwise.
 	 */
 	virtual bool AuthorizeSubscription( const IReceiveMessagesRef& Subscriber, const FName& TopicPattern ) = 0;
@@ -59,7 +52,6 @@ public:
 	 * Authorizes a request to unregister the specified recipient.
 	 *
 	 * @param Address The address of the recipient to unregister.
-	 *
 	 * @return true if the request was authorized, false otherwise.
 	 */
 	virtual bool AuthorizeUnregistration( const FMessageAddress& Address ) = 0;
@@ -69,7 +61,6 @@ public:
 	 *
 	 * @param Subscriber The subscriber.
 	 * @param TopicPattern The message topic pattern to unsubscribe from.
-	 *
 	 * @return true if the request is authorized, false otherwise.
 	 */
 	virtual bool AuthorizeUnsubscription( const IReceiveMessagesRef& Subscriber, const FName& TopicPattern ) = 0;
diff --git a/Engine/Source/Runtime/Messaging/Public/Interfaces/IInterceptMessages.h b/Engine/Source/Runtime/Messaging/Public/Interfaces/IInterceptMessages.h
--- a/Engine/Source/Runtime/Messaging/Public/Interfaces/IInterceptMessages.h
+++ b/Engine/Source/Runtime/Messaging/Public/Interfaces/IInterceptMessages.h
@@ -7,14 +7,10 @@
 #pragma once
 
 
-/**
- * Type definition for shared pointers to instances of IInterceptMessages.
- */
+/** Type definition for shared pointers to instances of IInterceptMessages. */
 typedef TSharedPtr<class IInterceptMessages, ESPMode::ThreadSafe> IInterceptMessagesPtr;
 
-/**
- * Type definition for shared references to instances of IInterceptMessages.
- */
+/** Type definition for shared references to instances of IInterceptMessages. */
 typedef TSharedRef<class IInterceptMessages, ESPMode::ThreadSafe> IInterceptMessagesRef;
 
 
@@ -29,7 +25,6 @@ public:
 	 * Intercepts a message before it is being passed to the message router.
 	 *
 	 * @param Context The context of the message to intercept.
-	 *
 	 * @return true if the message was intercepted and should not be routed, false otherwise.
 	 */
 	virtual bool InterceptMessage( const IMessageContextRef& Context ) = 0;
diff --git a/Engine/Source/Runtime/Messaging/Public/Interfaces/IMessageAttachment.h b/Engine/Source/Runtime/Messaging/Public/Interfaces/IMessageAttachment.h
--- a/Engine/Source/Runtime/Messaging/Public/Interfaces/IMessageAttachment.h
+++ b/Engine/Source/Runtime/Messaging/Public/Interfaces/IMessageAttachment.h
@@ -7,14 +7,10 @@
 #pragma once
 
 
-/**
- * Type definition for shared pointers to instances of IMessageAttachment.
- */
+/** Type definition for shared pointers to instances of IMessageAttachment. */
 typedef TSharedPtr<class IMessageAttachment, ESPMode::ThreadSafe> IMessageAttachmentPtr;
 
-/**
- * Type definition for shared references to instances of IMessageAttachment.
- */
+/** Type definition for shared references to instances of IMessageAttachment. */
 typedef TSharedRef<class IMessageAttachment, ESPMode::ThreadSafe> IMessageAttachmentRef;
 
 
diff --git a/Engine/Source/Runtime/Messaging/Public/Interfaces/IMessageBridge.h b/Engine/Source/Runtime/Messaging/Public/Interfaces/IMessageBridge.h
--- a/Engine/Source/Runtime/Messaging/Public/Interfaces/IMessageBridge.h
+++ b/Engine/Source/Runtime/Messaging/Public/Interfaces/IMessageBridge.h
@@ -7,14 +7,10 @@
 #pragma once
 
 
-/**
- * Type definition for shared pointers to instances of IMessageBridge.
- */
+/** Type definition for shared pointers to instances of IMessageBridge. */
 typedef TSharedPtr<class IMessageBridge, ESPMode::ThreadSafe> IMessageBridgePtr;
 
-/**
- * Type definition for shared references to instances of IMessageBridge.
- */
+/** Type definition for shared references to instances of IMessageBridge. */
 typedef TSharedRef<class IMessageBridge, ESPMode::ThreadSafe> IMessageBridgeRef;
 
 
@@ -31,8 +27,7 @@ public:
 	 * A disabled bridge will not receive any subscribed messages until it is enabled again.
 	 * Bridges should be created in a disabled state by default and explicitly enabled.
 	 *
-	 * @see Enable
-	 * @see IsEnabled
+	 * @see Enable, IsEnabled
 	 */
 	virtual void Disable( ) = 0;
 
@@ -42,8 +37,7 @@ public:
 	 * An activated bridge will receive subscribed messages.
 	 * Bridges should be created in a disabled state by default and explicitly enabled.
 	 *
-	 * @see Disable
-	 * @see IsEnabled
+	 * @see Disable, IsEnabled
 	 */
 	virtual void Enable( ) = 0;
 
@@ -51,9 +45,7 @@ public:
 	 * Checks whether the bridge is currently enabled.
 	 *
 	 * @return true if the bridge is enabled, false otherwise.
-	 *
-	 * @see Disable
-	 * @see Enable
+	 * @see Disable, Enable
 	 */
 	virtual bool IsEnabled( ) const = 0;
 
diff --git a/Engine/Source/Runtime/Messaging/Public/Interfaces/IMessageBus.h b/Engine/Source/Runtime/Messaging/Public/Interfaces/IMessageBus.h
--- a/Engine/Source/Runtime/Messaging/Public/Interfaces/IMessageBus.h
+++ b/Engine/Source/Runtime/Messaging/Public/Interfaces/IMessageBus.h
@@ -7,19 +7,13 @@
 #pragma once
 
 
-/**
- * Type definition for shared pointers to instances of IMessageBus.
- */
+/** Type definition for shared pointers to instances of IMessageBus. */
 typedef TSharedPtr<class IMessageBus, ESPMode::ThreadSafe> IMessageBusPtr;
 
-/**
- * Type definition for shared references to instances of IMessageBus.
- */
+/** Type definition for shared references to instances of IMessageBus. */
 typedef TSharedRef<class IMessageBus, ESPMode::ThreadSafe> IMessageBusRef;
 
-/**
- * Type definition for weak pointers to instances of IMessageBus.
- */
+/** Type definition for weak pointers to instances of IMessageBus. */
 typedef TWeakPtr<class IMessageBus, ESPMode::ThreadSafe> IMessageBusWeakPtr;
 
 
@@ -115,9 +109,7 @@ public:
 	 * @param Subscriber The subscriber wishing to receive the messages.
 	 * @param MessageType The type of messages to subscribe to (NAME_All = subscribe to all message types).
 	 * @param ScopeRange The range of message scopes to include in the subscription.
-	 *
 	 * @return The added subscription, or NULL if the subscription failed.
-	 *
 	 * @see Unsubscribe
 	 */
 	virtual IMessageSubscriptionPtr Subscribe( const IReceiveMessagesRef& Subscriber, const FName& MessageType, const FMessageScopeRange& ScopeRange ) = 0;
@@ -142,7 +134,6 @@ public:
 	 *
 	 * @param Subscriber The subscriber wishing to stop receiving the messages.
 	 * @param MessageType The type of messages to unsubscribe from (NAME_All = all types).
-	 *
 	 * @see Subscribe
 	 */
 	virtual void Unsubscribe( const IReceiveMessagesRef& Subscriber, const FName& MessageType ) = 0;
@@ -153,7 +144,6 @@ public:
 	 * Returns a delegate that is executed when the message bus is shutting down.
 	 *
 	 * @return The delegate.
-	 *
 	 * @see Shutdown
 	 */
 	virtual FOnMessageBusShutdown& OnShutdown( ) = 0;
diff --git a/Engine/Source/Runtime/Messaging/Public/Interfaces/IMessageContext.h b/Engine/Source/Runtime/Messaging/Public/Interfaces/IMessageContext.h
--- a/Engine/Source/Runtime/Messaging/Public/Interfaces/IMessageContext.h
+++ b/Engine/Source/Runtime/Messaging/Public/Interfaces/IMessageContext.h
@@ -26,19 +26,13 @@ namespace EMessageScope
 	 */
 	enum Type
 	{
-		/**
-		 * Deliver to subscribers in the same thread.
-		 */
+		/** Deliver to subscribers in the same thread. */
 		Thread,
 
-		/**
-		 * Deliver to subscribers in the same process.
-		 */
+		/** Deliver to subscribers in the same process. */
 		Process,
 
-		/**
-		 * Deliver to subscribers on the network.
-		 */
+		/** Deliver to subscribers on the network. */
 		Network,
 
 		/**
@@ -51,29 +45,19 @@ namespace EMessageScope
 }
 
 
-/**
- * Type definition for message endpoint identifiers.
- */
+/** Type definition for message endpoint identifiers. */
 typedef FGuid FMessageAddress;
 
-/**
- * Type definition for shared pointers to instances of IMessageContext.
- */
+/** Type definition for shared pointers to instances of IMessageContext. */
 typedef TSharedPtr<class IMessageContext, ESPMode::ThreadSafe> IMessageContextPtr;
 
-/**
- * Type definition for shared references to instances of IMessageContext.
- */
+/** Type definition for shared references to instances of IMessageContext. */
 typedef TSharedRef<class IMessageContext, ESPMode::ThreadSafe> IMessageContextRef;
 
-/**
- * Type definition for message scope ranges.
- */
+/** Type definition for message scope ranges. */
 typedef TRange<EMessageScope::Type> FMessageScopeRange;
 
-/**
- * Type definition for message scope range bounds.
- */
+/** Type definition for message scope range bounds. */
 typedef TRangeBound<EMessageScope::Type> FMessageScopeRangeBound;
 
 
@@ -102,7 +86,6 @@ public:
 	 * Gets the message address of the endpoint that forwarded this message.
 	 *
 	 * @return The forwarder's address.
-	 *
 	 * @see IsForwarded
 	 */
 	virtual const FMessageAddress& GetForwarder( ) const = 0;
@@ -188,7 +171,6 @@ public:
 	 * Checks whether this is a forwarded message.
 	 *
 	 * @return true if the message was forwarded, false otherwise.
-	 *
 	 * @see GetForwarder
 	 */
 	virtual bool IsForwarded( ) const = 0;
diff --git a/Engine/Source/Runtime/Messaging/Public/Interfaces/IMessageData.h b/Engine/Source/Runtime/Messaging/Public/Interfaces/IMessageData.h
--- a/Engine/Source/Runtime/Messaging/Public/Interfaces/IMessageData.h
+++ b/Engine/Source/Runtime/Messaging/Public/Interfaces/IMessageData.h
@@ -14,32 +14,22 @@ namespace EMessageDataState
 	 */
 	enum Type
 	{
-		/**
-		 * The message data is complete.
-		 */
+		/** The message data is complete. */
 		Complete,
 
-		/**
-		 * The message data is incomplete.
-		 */
+		/** The message data is incomplete. */
 		Incomplete,
 
-		/**
-		 * The message data is invalid.
-		 */
+		/** The message data is invalid. */
 		Invalid
 	};
 }
 
 
-/**
- * Type definition for shared pointers to instances of IMessageData.
- */
+/** Type definition for shared pointers to instances of IMessageData. */
 typedef TSharedPtr<class IMessageData, ESPMode::ThreadSafe> IMessageDataPtr;
 
-/**
- * Type definition for shared references to instances of IMessageData.
- */
+/** Type definition for shared references to instances of IMessageData. */
 typedef TSharedRef<class IMessageData, ESPMode::ThreadSafe> IMessageDataRef;
 
 
diff --git a/Engine/Source/Runtime/Messaging/Public/Interfaces/IMessageHandler.h b/Engine/Source/Runtime/Messaging/Public/Interfaces/IMessageHandler.h
--- a/Engine/Source/Runtime/Messaging/Public/Interfaces/IMessageHandler.h
+++ b/Engine/Source/Runtime/Messaging/Public/Interfaces/IMessageHandler.h
@@ -7,14 +7,10 @@
 #pragma once
 
 
-/**
- * Type definition for shared pointers to instances of IMessageHandler.
- */
+/** Type definition for shared pointers to instances of IMessageHandler. */
 typedef TSharedPtr<class IMessageHandler, ESPMode::ThreadSafe> IMessageHandlerPtr;
 
-/**
- * Type definition for shared references to instances of IMessageHandler.
- */
+/** Type definition for shared references to instances of IMessageHandler. */
 typedef TSharedRef<class IMessageHandler, ESPMode::ThreadSafe> IMessageHandlerRef;
 
 
diff --git a/Engine/Source/Runtime/Messaging/Public/Interfaces/IMessageSubscription.h b/Engine/Source/Runtime/Messaging/Public/Interfaces/IMessageSubscription.h
--- a/Engine/Source/Runtime/Messaging/Public/Interfaces/IMessageSubscription.h
+++ b/Engine/Source/Runtime/Messaging/Public/Interfaces/IMessageSubscription.h
@@ -7,14 +7,10 @@
 #pragma once
 
 
-/**
- * Type definition for shared pointers to instances of IMessageSubscription.
- */
+/** Type definition for shared pointers to instances of IMessageSubscription. */
 typedef TSharedPtr<class IMessageSubscription, ESPMode::ThreadSafe> IMessageSubscriptionPtr;
 
-/**
- * Type definition for shared references to instances of IMessageSubscription.
- */
+/** Type definition for shared references to instances of IMessageSubscription. */
 typedef TSharedRef<class IMessageSubscription, ESPMode::ThreadSafe> IMessageSubscriptionRef;
 
 
diff --git a/Engine/Source/Runtime/Messaging/Public/Interfaces/IMessageTracer.h b/Engine/Source/Runtime/Messaging/Public/Interfaces/IMessageTracer.h
--- a/Engine/Source/Runtime/Messaging/Public/Interfaces/IMessageTracer.h
+++ b/Engine/Source/Runtime/Messaging/Public/Interfaces/IMessageTracer.h
@@ -7,64 +7,40 @@
 #pragma once
 
 
-/**
- * Type definition for shared pointers to instances of IMessageTracer.
- */
+/** Type definition for shared pointers to instances of IMessageTracer. */
 typedef TSharedPtr<class IMessageTracer, ESPMode::ThreadSafe> IMessageTracerPtr;
 
-/**
- * Type definition for shared references to instances of IMessageTracer.
- */
+/** Type definition for shared references to instances of IMessageTracer. */
 typedef TSharedRef<class IMessageTracer, ESPMode::ThreadSafe> IMessageTracerRef;
 
-/**
- * Type definition for shared pointers to instances of FMessageTracerAddressInfo.
- */
+/** Type definition for shared pointers to instances of FMessageTracerAddressInfo. */
 typedef TSharedPtr<struct FMessageTracerAddressInfo> FMessageTracerAddressInfoPtr;
 
-/**
- * Type definition for shared references to instances of FMessageTracerAddressInfo.
- */
+/** Type definition for shared references to instances of FMessageTracerAddressInfo. */
 typedef TSharedRef<struct FMessageTracerAddressInfo> FMessageTracerAddressInfoRef;
 
-/**
- * Type definition for shared pointers to instances of FMessageTracerDispatchState.
- */
+/** Type definition for shared pointers to instances of FMessageTracerDispatchState. */
 typedef TSharedPtr<struct FMessageTracerDispatchState> FMessageTracerDispatchStatePtr;
 
-/**
- * Type definition for shared references to instances of FMessageTracerDispatchState.
- */
+/** Type definition for shared references to instances of FMessageTracerDispatchState. */
 typedef TSharedRef<struct FMessageTracerDispatchState> FMessageTracerDispatchStateRef;
 
-/**
- * Type definition for shared pointers to instances of FMessageTracerEndpointInfo.
- */
+/** Type definition for shared pointers to instances of FMessageTracerEndpointInfo. */
 typedef TSharedPtr<struct FMessageTracerEndpointInfo> FMessageTracerEndpointInfoPtr;
 
-/**
- * Type definition for shared references to instances of FMessageTracerEndpointInfo.
- */
+/** Type definition for shared references to instances of FMessageTracerEndpointInfo. */
 typedef TSharedRef<struct FMessageTracerEndpointInfo> FMessageTracerEndpointInfoRef;
 
-/**
- * Type definition for shared pointers to instances of FMessageTracerMessageInfo.
- */
+/** Type definition for shared pointers to instances of FMessageTracerMessageInfo. */
 typedef TSharedPtr<struct FMessageTracerMessageInfo> FMessageTracerMessageInfoPtr;
 
-/**
- * Type definition for shared references to instances of FMessageTracerMessageInfo.
- */
+/** Type definition for shared references to instances of FMessageTracerMessageInfo. */
 typedef TSharedRef<struct FMessageTracerMessageInfo> FMessageTracerMessageInfoRef;
 
-/**
- * Type definition for shared pointers to instances of FMessageTracerTypeInfo.
- */
+/** Type definition for shared pointers to instances of FMessageTracerTypeInfo. */
 typedef TSharedPtr<struct FMessageTracerTypeInfo> FMessageTracerTypeInfoPtr;
 
-/**
- * Type definition for shared references to instances of FMessageTracerTypeInfo.
- */
+/** Type definition for shared references to instances of FMessageTracerTypeInfo. */
 typedef TSharedRef<struct FMessageTracerTypeInfo> FMessageTracerTypeInfoRef;
 
 
@@ -75,19 +51,13 @@ namespace EMessageTracerDispatchTypes
 	 */
 	enum Type
 	{
-		/**
-		 * The message is being dispatched directly.
-		 */
+		/** The message is being dispatched directly. */
 		Direct,
 
-		/**
-		 * The message hasn't been dispatched yet.
-		 */
+		/** The message hasn't been dispatched yet. */
 		Pending,
 
-		/**
-		 * The message is being dispatched using the task graph system.
-		 */
+		/** The message is being dispatched using the task graph system. */
 		TaskGraph
 	};
 }
@@ -98,29 +68,19 @@ namespace EMessageTracerDispatchTypes
  */
 struct FMessageTracerDispatchState
 {
-	/**
-	 * Holds the dispatch latency (in seconds).
-	 */
+	/** Holds the dispatch latency (in seconds). */
 	double DispatchLatency;
 
-	/**
-	 * Holds the message's dispatch type for the specified endpoint.
-	 */
+	/** Holds the message's dispatch type for the specified endpoint. */
 	EMessageTracerDispatchTypes::Type DispatchType;
 
-	/**
-	 * Holds the endpoint to which the message was or is being dispatched.
-	 */
+	/** Holds the endpoint to which the message was or is being dispatched. */
 	FMessageTracerEndpointInfoPtr EndpointInfo;
 
-	/**
-	 * Holds the time at which the message was dispatched.
-	 */
+	/** Holds the time at which the message was dispatched. */
 	double TimeDispatched;
 
-	/**
-	 * Holds the time at which the message was actually handled (0.0 = not handled yet).
-	 */
+	/** Holds the time at which the message was actually handled (0.0 = not handled yet). */
 	double TimeHandled;
 };
 
@@ -130,19 +90,13 @@ struct FMessageTracerDispatchState
  */
 struct FMessageTracerAddressInfo
 {
-	/**
-	 * Holds a recipient's message address.
-	 */
+	/** Holds a recipient's message address. */
 	FMessageAddress Address;
 
-	/**
-	 * Holds the time at which this address was registered.
-	 */
+	/** Holds the time at which this address was registered. */
 	double TimeRegistered;
 
-	/**
-	 * Holds the time at which this address was unregistered.
-	 */
+	/** Holds the time at which this address was unregistered. */
 	double TimeUnregistered;
 };
 
@@ -152,29 +106,19 @@ struct FMessageTracerAddressInfo
  */
 struct FMessageTracerEndpointInfo
 {
-	/**
-	 * Holds the recipient's address information.
-	 */
+	/** Holds the recipient's address information. */
 	TMap<FMessageAddress, FMessageTracerAddressInfoPtr> AddressInfos;
 
-	/**
-	 * Holds the recipient's human readable name.
-	 */
+	/** Holds the recipient's human readable name. */
 	FName Name;
 
-	/**
-	 * Holds the list of messages received by this recipient.
-	 */
+	/** Holds the list of messages received by this recipient. */
 	TArray<FMessageTracerMessageInfoPtr> ReceivedMessages;
 
-	/**
-	 * Holds a flag indicating whether this is a remote recipient.
-	 */
+	/** Holds a flag indicating whether this is a remote recipient. */
 	bool Remote;
 
-	/**
-	 * Holds the list of messages sent by this recipient.
-	 */
+	/** Holds the list of messages sent by this recipient. */
 	TArray<FMessageTracerMessageInfoPtr> SentMessages;
 };
 
@@ -184,34 +128,22 @@ struct FMessageTracerEndpointInfo
  */
 struct FMessageTracerMessageInfo
 {
-	/**
-	 * Holds a pointer to the message context.
-	 */
+	/** Holds a pointer to the message context. */
 	IMessageContextPtr Context;
 
-	/**
-	 * Holds the message's dispatch states per endpoint.
-	 */
+	/** Holds the message's dispatch states per endpoint. */
 	TMap<FMessageTracerEndpointInfoPtr, FMessageTracerDispatchStatePtr> DispatchStates;
 
-	/**
-	 * Pointer to the sender's endpoint information.
-	 */
+	/** Pointer to the sender's endpoint information. */
 	FMessageTracerEndpointInfoPtr SenderInfo;
 
-	/**
-	 * Holds the time at which the message was routed (0.0 = pending).
-	 */
+	/** Holds the time at which the message was routed (0.0 = pending). */
 	double TimeRouted;
 
-	/**
-	 * Holds the time at which the message was sent.
-	 */
+	/** Holds the time at which the message was sent. */
 	double TimeSent;
 
-	/**
-	 * Pointer to the message's type information.
-	 */
+	/** Pointer to the message's type information. */
 	FMessageTracerTypeInfoPtr TypeInfo;
 };
 
@@ -221,14 +153,10 @@ struct FMessageTracerMessageInfo
  */
 struct FMessageTracerTypeInfo
 {
-	/**
-	 * Holds the collection of messages of this type.
-	 */
+	/** Holds the collection of messages of this type. */
 	TArray<FMessageTracerMessageInfoPtr> Messages;
 
-	/**
-	 * Holds a name of the message type.
-	 */
+	/** Holds a name of the message type. */
 	FName TypeName;
 };
 
@@ -258,8 +186,7 @@ public:
 	/**
 	 * Breaks message routing.
 	 *
-	 * @see Continue
-	 * @see Step
+	 * @see Continue, Step
 	 */
 	virtual void Break( ) = 0;
 
@@ -280,8 +207,7 @@ public:
 	 *
 	 * @return true if the tracer is running, false otherwise.
 	 *
-	 * @see Start
-	 * @see Stop
+	 * @see Start, Stop
 	 */
 	virtual bool IsRunning( ) const = 0;
 
@@ -293,24 +219,21 @@ public:
 	/**
 	 * Starts the tracer.
 	 *
-	 * @see IsRunning
-	 * @see Stop
+	 * @see IsRunning, Stop
 	 */
 	virtual void Start( ) = 0;
 
 	/**
 	 * Steps the tracer to the next message.
 	 *
-	 * @see Break
-	 * @see Continue
+	 * @see Break, Continue
 	 */
 	virtual void Step( ) = 0;
 
 	/**
 	 * Stops the tracer.
 	 *
-	 * @see IsRunning
-	 * @see Start
+	 * @see IsRunning, Start
 	 */
 	virtual void Stop( ) = 0;
 
@@ -318,7 +241,6 @@ public:
 	 * Ticks the tracer.
 	 *
 	 * @param DeltaTime The time in seconds since the last tick.
-	 *
 	 * @return true if any events were processed.
 	 */
 	virtual bool Tick( float DeltaTime ) = 0;
@@ -329,7 +251,6 @@ public:
 	 * Gets the list of known message endpoints.
 	 *
 	 * @param OutEndpoints Will contain the list of endpoints.
-	 *
 	 * @return The number of endpoints returned.
 	 */
 	virtual int32 GetEndpoints( TArray<FMessageTracerEndpointInfoPtr>& OutEndpoints ) const = 0;
@@ -346,7 +267,6 @@ public:
 	 *
 	 * @param NameFilter The name substring to filter with.
 	 * @param OutTypes Will contain the list of message types.
-	 *
 	 * @return The number of message types returned.
 	 */
 	virtual int32 GetMessageTypes( TArray<FMessageTracerTypeInfoPtr>& OutTypes ) const = 0;
diff --git a/Engine/Source/Runtime/Messaging/Public/Interfaces/IMessageTracerBreakpoint.h b/Engine/Source/Runtime/Messaging/Public/Interfaces/IMessageTracerBreakpoint.h
--- a/Engine/Source/Runtime/Messaging/Public/Interfaces/IMessageTracerBreakpoint.h
+++ b/Engine/Source/Runtime/Messaging/Public/Interfaces/IMessageTracerBreakpoint.h
@@ -7,14 +7,10 @@
 #pragma once
 
 
-/**
- * Type definition for shared pointers to instances of IMessageTracerBreakpoint.
- */
+/** Type definition for shared pointers to instances of IMessageTracerBreakpoint. */
 typedef TSharedPtr<class IMessageTracerBreakpoint, ESPMode::ThreadSafe> IMessageTracerBreakpointPtr;
 
-/**
- * Type definition for shared references to instances of IMessageTracerBreakpoint.
- */
+/** Type definition for shared references to instances of IMessageTracerBreakpoint. */
 typedef TSharedRef<class IMessageTracerBreakpoint, ESPMode::ThreadSafe> IMessageTracerBreakpointRef;
 
 
@@ -36,7 +32,6 @@ public:
 	 * Checks whether the tracer should break on the given message.
 	 *
 	 * @param Context The message to break on.
-	 *
 	 * @return true if the tracer should break, false otherwise.
 	 */
 	virtual bool ShouldBreak( const IMessageContextRef& Context ) const = 0;
diff --git a/Engine/Source/Runtime/Messaging/Public/Interfaces/IMessagingModule.h b/Engine/Source/Runtime/Messaging/Public/Interfaces/IMessagingModule.h
--- a/Engine/Source/Runtime/Messaging/Public/Interfaces/IMessagingModule.h
+++ b/Engine/Source/Runtime/Messaging/Public/Interfaces/IMessagingModule.h
@@ -25,9 +25,7 @@ public:
 	 * @param Bus The message bus to attach the bridge to.
 	 * @param Serializer The message serializer to use.
 	 * @param Transport The message transport technology to use.
-	 *
 	 * @return The new message bridge, or NULL if the bridge couldn't be created.
-	 *
 	 * @see CreateBus
 	 */
 	virtual IMessageBridgePtr CreateBridge( const FMessageAddress& Address, const IMessageBusRef& Bus, const ISerializeMessagesRef& Serializer, const ITransportMessagesRef& Transport ) = 0;
@@ -36,9 +34,7 @@ public:
 	 * Creates a new message bus.
 	 *
 	 * @param RecipientAuthorizer An optional recipient authorizer.
-	 *
 	 * @return The new message bus, or NULL if the bus couldn't be created.
-	 *
 	 * @see CreateBridge
 	 */
 	virtual IMessageBusPtr CreateBus( const IAuthorizeMessageRecipientsPtr& RecipientAuthorizer ) = 0;
diff --git a/Engine/Source/Runtime/Messaging/Public/Interfaces/IMutableMessageContext.h b/Engine/Source/Runtime/Messaging/Public/Interfaces/IMutableMessageContext.h
--- a/Engine/Source/Runtime/Messaging/Public/Interfaces/IMutableMessageContext.h
+++ b/Engine/Source/Runtime/Messaging/Public/Interfaces/IMutableMessageContext.h
@@ -7,14 +7,10 @@
 #pragma once
 
 
-/**
- * Type definition for shared pointers to instances of IMutableMessageContext.
- */
+/** Type definition for shared pointers to instances of IMutableMessageContext. */
 typedef TSharedPtr<class IMutableMessageContext, ESPMode::ThreadSafe> IMutableMessageContextPtr;
 
-/**
- * Type definition for shared references to instances of IMutableMessageContext.
- */
+/** Type definition for shared references to instances of IMutableMessageContext. */
 typedef TSharedRef<class IMutableMessageContext, ESPMode::ThreadSafe> IMutableMessageContextRef;
 
 
diff --git a/Engine/Source/Runtime/Messaging/Public/Interfaces/IReceiveMessages.h b/Engine/Source/Runtime/Messaging/Public/Interfaces/IReceiveMessages.h
--- a/Engine/Source/Runtime/Messaging/Public/Interfaces/IReceiveMessages.h
+++ b/Engine/Source/Runtime/Messaging/Public/Interfaces/IReceiveMessages.h
@@ -7,19 +7,13 @@
 #pragma once
 
 
-/**
- * Type definition for shared pointers to instances of IReceiveMessages.
- */
+/** Type definition for shared pointers to instances of IReceiveMessages. */
 typedef TSharedPtr<class IReceiveMessages, ESPMode::ThreadSafe> IReceiveMessagesPtr;
 
-/**
- * Type definition for shared references to instances of IReceiveMessages.
- */
+/** Type definition for shared references to instances of IReceiveMessages. */
 typedef TSharedRef<class IReceiveMessages, ESPMode::ThreadSafe> IReceiveMessagesRef;
 
-/**
- * Type definition for shared pointers to instances of IReceiveMessages.
- */
+/** Type definition for shared pointers to instances of IReceiveMessages. */
 typedef TWeakPtr<class IReceiveMessages, ESPMode::ThreadSafe> IReceiveMessagesWeakPtr;
 
 
@@ -39,7 +33,6 @@ public:
 	 * Gets the recipient's name (for debugging purposes).
 	 *
 	 * @return The debug name.
-	 *
 	 * @see GetRecipientId
 	 */
 	virtual FName GetDebugName( ) const = 0;
@@ -48,7 +41,6 @@ public:
 	 * Gets the recipient's unique identifier (for debugging purposes).
 	 *
 	 * @return The recipient's identifier.
-	 *
 	 * @see GetRecipientName
 	 */
 	virtual const FGuid& GetRecipientId( ) const = 0;
diff --git a/Engine/Source/Runtime/Messaging/Public/Interfaces/ISendMessages.h b/Engine/Source/Runtime/Messaging/Public/Interfaces/ISendMessages.h
--- a/Engine/Source/Runtime/Messaging/Public/Interfaces/ISendMessages.h
+++ b/Engine/Source/Runtime/Messaging/Public/Interfaces/ISendMessages.h
@@ -7,14 +7,10 @@
 #pragma once
 
 
-/**
- * Type definition for shared pointers to instances of ISendMessages.
- */
+/** Type definition for shared pointers to instances of ISendMessages. */
 typedef TSharedPtr<class ISendMessages, ESPMode::ThreadSafe> ISendMessagesPtr;
 
-/**
- * Type definition for shared references to instances of ISendMessages.
- */
+/** Type definition for shared references to instances of ISendMessages. */
 typedef TSharedRef<class ISendMessages, ESPMode::ThreadSafe> ISendMessagesRef;
 
 
diff --git a/Engine/Source/Runtime/Messaging/Public/Interfaces/ISerializeMessages.h b/Engine/Source/Runtime/Messaging/Public/Interfaces/ISerializeMessages.h
--- a/Engine/Source/Runtime/Messaging/Public/Interfaces/ISerializeMessages.h
+++ b/Engine/Source/Runtime/Messaging/Public/Interfaces/ISerializeMessages.h
@@ -7,14 +7,10 @@
 #pragma once
 
 
-/**
- * Type definition for shared pointers to instances of ISerializeMessages.
- */
+/** Type definition for shared pointers to instances of ISerializeMessages. */
 typedef TSharedPtr<class ISerializeMessages, ESPMode::ThreadSafe> ISerializeMessagesPtr;
 
-/**
- * Type definition for shared references to instances of ISerializeMessages.
- */
+/** Type definition for shared references to instances of ISerializeMessages. */
 typedef TSharedRef<class ISerializeMessages, ESPMode::ThreadSafe> ISerializeMessagesRef;
 
 
@@ -33,7 +29,6 @@ public:
 	 *
 	 * @param Data The archive to serialize from.
 	 * @param OutContext Will hold the context of the deserialized message.
-	 *
 	 * @return true if deserialization was successful, false otherwise.
 	 */
 	virtual bool DeserializeMessage( FArchive& Archive, IMutableMessageContextRef& OutContext ) = 0;
@@ -43,7 +38,6 @@ public:
 	 *
 	 * @param Context The context of the message to serialize.
 	 * @param Archive The archive to serialize into.
-	 *
 	 * @return true if serialization was successful, false otherwise.
 	 */
 	virtual bool SerializeMessage( const IMessageContextRef& Context, FArchive& Archive ) = 0;
diff --git a/Engine/Source/Runtime/Messaging/Public/Interfaces/ITransportMessages.h b/Engine/Source/Runtime/Messaging/Public/Interfaces/ITransportMessages.h
--- a/Engine/Source/Runtime/Messaging/Public/Interfaces/ITransportMessages.h
+++ b/Engine/Source/Runtime/Messaging/Public/Interfaces/ITransportMessages.h
@@ -7,14 +7,10 @@
 #pragma once
 
 
-/**
- * Type definition for shared pointers to instances of ITransportMessages.
- */
+/** Type definition for shared pointers to instances of ITransportMessages. */
 typedef TSharedPtr<class ITransportMessages, ESPMode::ThreadSafe> ITransportMessagesPtr;
 
-/**
- * Type definition for shared references to instances of ITransportMessages.
- */
+/** Type definition for shared references to instances of ITransportMessages. */
 typedef TSharedRef<class ITransportMessages, ESPMode::ThreadSafe> ITransportMessagesRef;
 
 
@@ -66,7 +62,6 @@ public:
 	 * Starts up the message transport.
 	 *
 	 * @return Whether the transport was started successfully.
-	 *
 	 * @see StopTransport
 	 */
 	virtual bool StartTransport( ) = 0;
@@ -84,7 +79,6 @@ public:
 	 * @param Data The serialized message data to transport.
 	 * @param Attachment An optional message attachment (i.e. file or memory buffer).
 	 * @param Recipients The transport nodes to send the message to.
-	 *
 	 * @return true if the message is being transported, false otherwise.
 	 */
 	virtual bool TransportMessage( const IMessageDataRef& Data, const IMessageAttachmentPtr& Attachment, const TArray<FGuid>& Recipients ) = 0;
