renamed OVERRIDE to override

diff --git a/Engine/Source/Runtime/Messaging/Private/Bridge/MessageBridge.h b/Engine/Source/Runtime/Messaging/Private/Bridge/MessageBridge.h
--- a/Engine/Source/Runtime/Messaging/Private/Bridge/MessageBridge.h
+++ b/Engine/Source/Runtime/Messaging/Private/Bridge/MessageBridge.h
@@ -37,11 +37,11 @@ public:
 
 	// Begin IMessageBridge interface
 
-	virtual void Disable( ) OVERRIDE;
+	virtual void Disable( ) override;
 
-	virtual void Enable( ) OVERRIDE;
+	virtual void Enable( ) override;
 
-	virtual bool IsEnabled( ) const OVERRIDE
+	virtual bool IsEnabled( ) const override
 	{
 		return Enabled;
 	}
@@ -52,27 +52,27 @@ public:
 
 	// Begin IReceiveMessages interface
 
-	virtual FName GetDebugName( ) const OVERRIDE
+	virtual FName GetDebugName( ) const override
 	{
 		return *FString::Printf(TEXT("FMessageBridge (%s)"), *Transport->GetDebugName().ToString());
 	}
 
-	virtual const FGuid& GetRecipientId( ) const OVERRIDE
+	virtual const FGuid& GetRecipientId( ) const override
 	{
 		return Id;
 	}
 
-	virtual ENamedThreads::Type GetRecipientThread( ) const OVERRIDE
+	virtual ENamedThreads::Type GetRecipientThread( ) const override
 	{
 		return ENamedThreads::AnyThread;
 	}
 
-	virtual bool IsLocal( ) const OVERRIDE
+	virtual bool IsLocal( ) const override
 	{
 		return false;
 	}
 
-	virtual void ReceiveMessage( const IMessageContextRef& Context ) OVERRIDE;
+	virtual void ReceiveMessage( const IMessageContextRef& Context ) override;
 
 	// End IReceiveMessages interface
 
@@ -80,12 +80,12 @@ public:
 
 	// Begin ISendMessages interface
 
-	virtual FMessageAddress GetSenderAddress( ) OVERRIDE
+	virtual FMessageAddress GetSenderAddress( ) override
 	{
 		return Address;
 	}
 
-	virtual void NotifyMessageError( const IMessageContextRef& Context, const FString& Error ) OVERRIDE;
+	virtual void NotifyMessageError( const IMessageContextRef& Context, const FString& Error ) override;
 
 	// End ISendMessages interface
 
diff --git a/Engine/Source/Runtime/Messaging/Private/Bridge/MessageData.h b/Engine/Source/Runtime/Messaging/Private/Bridge/MessageData.h
--- a/Engine/Source/Runtime/Messaging/Private/Bridge/MessageData.h
+++ b/Engine/Source/Runtime/Messaging/Private/Bridge/MessageData.h
@@ -48,17 +48,17 @@ public:
 
 	// Begin IMessageAttachment interface
 
-	virtual FArchive* CreateReader( ) OVERRIDE
+	virtual FArchive* CreateReader( ) override
 	{
 		return new FMemoryReader(Data, true);
 	}
 
-	virtual EMessageDataState::Type GetState( ) const OVERRIDE
+	virtual EMessageDataState::Type GetState( ) const override
 	{
 		return State;
 	}
 
-	virtual FSimpleDelegate& OnStateChanged( ) OVERRIDE
+	virtual FSimpleDelegate& OnStateChanged( ) override
 	{
 		return StateChangedDelegate;
 	}
diff --git a/Engine/Source/Runtime/Messaging/Private/Bus/MessageBus.h b/Engine/Source/Runtime/Messaging/Private/Bus/MessageBus.h
--- a/Engine/Source/Runtime/Messaging/Private/Bus/MessageBus.h
+++ b/Engine/Source/Runtime/Messaging/Private/Bus/MessageBus.h
@@ -32,17 +32,17 @@ public:
 
 	// Begin IMessageBus interface
 
-	virtual void Forward( const IMessageContextRef& Context, const TArray<FMessageAddress>& Recipients, EMessageScope::Type ForwardingScope, const FTimespan& Delay, const ISendMessagesRef& Forwarder ) OVERRIDE
+	virtual void Forward( const IMessageContextRef& Context, const TArray<FMessageAddress>& Recipients, EMessageScope::Type ForwardingScope, const FTimespan& Delay, const ISendMessagesRef& Forwarder ) override
 	{
 		Router->RouteMessage(MakeShareable(new FMessageContext(Context, Forwarder->GetSenderAddress(), Recipients, ForwardingScope, FDateTime::UtcNow(), FTaskGraphInterface::Get().GetCurrentThreadIfKnown())));
 	}
 
-	virtual IMessageTracerRef GetTracer( ) OVERRIDE
+	virtual IMessageTracerRef GetTracer( ) override
 	{
 		return Router->GetTracer();
 	}
 
-	virtual void Intercept( const IInterceptMessagesRef& Interceptor, const FName& MessageType ) OVERRIDE
+	virtual void Intercept( const IInterceptMessagesRef& Interceptor, const FName& MessageType ) override
 	{
 		if (MessageType != NAME_None)
 		{
@@ -53,29 +53,29 @@ public:
 		}
 	}
 
-	virtual FOnMessageBusShutdown& OnShutdown( ) OVERRIDE
+	virtual FOnMessageBusShutdown& OnShutdown( ) override
 	{
 		return ShutdownDelegate;
 	}
 
-	virtual void Publish( void* Message, UScriptStruct* TypeInfo, EMessageScope::Type Scope, const FTimespan& Delay, const FDateTime& Expiration, const ISendMessagesRef& Publisher ) OVERRIDE
+	virtual void Publish( void* Message, UScriptStruct* TypeInfo, EMessageScope::Type Scope, const FTimespan& Delay, const FDateTime& Expiration, const ISendMessagesRef& Publisher ) override
 	{
 		Router->RouteMessage(MakeShareable(new FMessageContext(Message, TypeInfo, nullptr, Publisher->GetSenderAddress(), TArray<FMessageAddress>(), Scope, FDateTime::UtcNow() + Delay, Expiration, FTaskGraphInterface::Get().GetCurrentThreadIfKnown())));
 	}
 
-	virtual void Register( const FMessageAddress& Address, const IReceiveMessagesRef& Recipient ) OVERRIDE
+	virtual void Register( const FMessageAddress& Address, const IReceiveMessagesRef& Recipient ) override
 	{
 		Router->AddRecipient(Address, Recipient);
 	}
 
-	virtual void Send( void* Message, UScriptStruct* TypeInfo, const IMessageAttachmentPtr& Attachment, const TArray<FMessageAddress>& Recipients, const FTimespan& Delay, const FDateTime& Expiration, const ISendMessagesRef& Sender ) OVERRIDE
+	virtual void Send( void* Message, UScriptStruct* TypeInfo, const IMessageAttachmentPtr& Attachment, const TArray<FMessageAddress>& Recipients, const FTimespan& Delay, const FDateTime& Expiration, const ISendMessagesRef& Sender ) override
 	{
 		Router->RouteMessage(MakeShareable(new FMessageContext(Message, TypeInfo, Attachment, Sender->GetSenderAddress(), Recipients, EMessageScope::Network, FDateTime::UtcNow() + Delay, Expiration, FTaskGraphInterface::Get().GetCurrentThreadIfKnown())));
 	}
 
-	virtual void Shutdown( ) OVERRIDE;
+	virtual void Shutdown( ) override;
 
-	virtual IMessageSubscriptionPtr Subscribe( const IReceiveMessagesRef& Subscriber, const FName& MessageType, const FMessageScopeRange& ScopeRange ) OVERRIDE
+	virtual IMessageSubscriptionPtr Subscribe( const IReceiveMessagesRef& Subscriber, const FName& MessageType, const FMessageScopeRange& ScopeRange ) override
 	{
 		if (MessageType != NAME_None)
 		{
@@ -91,7 +91,7 @@ public:
 		return nullptr;
 	}
 
-	virtual void Unintercept( const IInterceptMessagesRef& Interceptor, const FName& MessageType ) OVERRIDE
+	virtual void Unintercept( const IInterceptMessagesRef& Interceptor, const FName& MessageType ) override
 	{
 		if (MessageType != NAME_None)
 		{
@@ -99,7 +99,7 @@ public:
 		}
 	}
 
-	virtual void Unregister( const FMessageAddress& Address ) OVERRIDE
+	virtual void Unregister( const FMessageAddress& Address ) override
 	{
 		if (!RecipientAuthorizer.IsValid() || RecipientAuthorizer->AuthorizeUnregistration(Address))
 		{
@@ -107,7 +107,7 @@ public:
 		}
 	}
 
-	virtual void Unsubscribe( const IReceiveMessagesRef& Subscriber, const FName& MessageType ) OVERRIDE
+	virtual void Unsubscribe( const IReceiveMessagesRef& Subscriber, const FName& MessageType ) override
 	{
 		if (MessageType != NAME_None)
 		{
diff --git a/Engine/Source/Runtime/Messaging/Private/Bus/MessageContext.h b/Engine/Source/Runtime/Messaging/Private/Bus/MessageContext.h
--- a/Engine/Source/Runtime/Messaging/Private/Bus/MessageContext.h
+++ b/Engine/Source/Runtime/Messaging/Private/Bus/MessageContext.h
@@ -89,7 +89,7 @@ public:
 
 	// Begin IMessageContext interface
 
-	virtual IMessageAttachmentPtr GetAttachment( ) const OVERRIDE
+	virtual IMessageAttachmentPtr GetAttachment( ) const override
 	{
 		if (OriginalContext.IsValid())
 		{
@@ -99,7 +99,7 @@ public:
 		return Attachment;
 	}
 
-	virtual const FDateTime& GetExpiration( ) const OVERRIDE
+	virtual const FDateTime& GetExpiration( ) const override
 	{
 		if (OriginalContext.IsValid())
 		{
@@ -109,12 +109,12 @@ public:
 		return Expiration;
 	}
 
-	virtual const FMessageAddress& GetForwarder( ) const OVERRIDE
+	virtual const FMessageAddress& GetForwarder( ) const override
 	{
 		return Forwarder;
 	}
 
-	virtual const TMap<FName, FString>& GetHeaders( ) const OVERRIDE
+	virtual const TMap<FName, FString>& GetHeaders( ) const override
 	{
 		if (OriginalContext.IsValid())
 		{
@@ -124,7 +124,7 @@ public:
 		return Headers;
 	}
 
-	virtual const void* GetMessage( ) const OVERRIDE
+	virtual const void* GetMessage( ) const override
 	{
 		if (OriginalContext.IsValid())
 		{
@@ -134,7 +134,7 @@ public:
 		return Message;
 	}
 
-	virtual FName GetMessageType( ) const OVERRIDE
+	virtual FName GetMessageType( ) const override
 	{
 		if (IsValid())
 		{
@@ -144,7 +144,7 @@ public:
 		return NAME_None;
 	}
 
-	virtual const TWeakObjectPtr<UScriptStruct>& GetMessageTypeInfo( ) const OVERRIDE
+	virtual const TWeakObjectPtr<UScriptStruct>& GetMessageTypeInfo( ) const override
 	{
 		if (OriginalContext.IsValid())
 		{
@@ -154,22 +154,22 @@ public:
 		return TypeInfo;
 	}
 
-	virtual IMessageContextPtr GetOriginalContext( ) const OVERRIDE
+	virtual IMessageContextPtr GetOriginalContext( ) const override
 	{
 		return OriginalContext;
 	}
 
-	virtual const TArray<FMessageAddress>& GetRecipients( ) const OVERRIDE
+	virtual const TArray<FMessageAddress>& GetRecipients( ) const override
 	{
 		return Recipients;
 	}
 
-	virtual EMessageScope::Type GetScope( ) const OVERRIDE
+	virtual EMessageScope::Type GetScope( ) const override
 	{
 		return Scope;
 	}
 
-	virtual const FMessageAddress& GetSender( ) const OVERRIDE
+	virtual const FMessageAddress& GetSender( ) const override
 	{
 		if (OriginalContext.IsValid())
 		{
@@ -179,17 +179,17 @@ public:
 		return Sender;
 	}
 
-	virtual ENamedThreads::Type GetSenderThread( ) const OVERRIDE
+	virtual ENamedThreads::Type GetSenderThread( ) const override
 	{
 		return SenderThread;
 	}
 
-	virtual const FDateTime& GetTimeForwarded( ) const OVERRIDE
+	virtual const FDateTime& GetTimeForwarded( ) const override
 	{
 		return TimeForwarded;
 	}
 
-	virtual const FDateTime& GetTimeSent( ) const OVERRIDE
+	virtual const FDateTime& GetTimeSent( ) const override
 	{
 		if (OriginalContext.IsValid())
 		{
@@ -199,12 +199,12 @@ public:
 		return TimeSent;
 	}
 
-	virtual bool IsForwarded( ) const OVERRIDE
+	virtual bool IsForwarded( ) const override
 	{
 		return OriginalContext.IsValid();
 	}
 
-	virtual bool IsValid( ) const OVERRIDE
+	virtual bool IsValid( ) const override
 	{
 		if (OriginalContext.IsValid())
 		{
@@ -220,43 +220,43 @@ public:
 
 	// Begin IMutableMessageContext interface
 
-	virtual void AddRecipient( const FMessageAddress& Recipient ) OVERRIDE
+	virtual void AddRecipient( const FMessageAddress& Recipient ) override
 	{
 		Recipients.Add(Recipient);
 	}
 
-	virtual void SetAttachment( const IMessageAttachmentPtr& InAttachment ) OVERRIDE
+	virtual void SetAttachment( const IMessageAttachmentPtr& InAttachment ) override
 	{
 		Attachment = InAttachment;
 	}
 
-	virtual void SetMessage( void* InMessage, UScriptStruct* InTypeInfo ) OVERRIDE
+	virtual void SetMessage( void* InMessage, UScriptStruct* InTypeInfo ) override
 	{
 		Message = InMessage;
 		TypeInfo = InTypeInfo;
 	}
 
-	virtual void SetExpiration( const FDateTime& InExpiration ) OVERRIDE
+	virtual void SetExpiration( const FDateTime& InExpiration ) override
 	{
 		Expiration = InExpiration;
 	}
 
-	virtual void SetHeader( const FName& Key, const FString& Value ) OVERRIDE
+	virtual void SetHeader( const FName& Key, const FString& Value ) override
 	{
 		Headers.Add(Key, Value);
 	}
 
-	virtual void SetScope( EMessageScope::Type InScope ) OVERRIDE
+	virtual void SetScope( EMessageScope::Type InScope ) override
 	{
 		Scope = InScope;
 	}
 
-	virtual void SetSender( const FMessageAddress& InSender ) OVERRIDE
+	virtual void SetSender( const FMessageAddress& InSender ) override
 	{
 		Sender = InSender;
 	}
 
-	virtual void SetTimeSent( const FDateTime& InTimeSent ) OVERRIDE
+	virtual void SetTimeSent( const FDateTime& InTimeSent ) override
 	{
 		TimeSent = InTimeSent;
 	}
diff --git a/Engine/Source/Runtime/Messaging/Private/Bus/MessageRouter.h b/Engine/Source/Runtime/Messaging/Private/Bus/MessageRouter.h
--- a/Engine/Source/Runtime/Messaging/Private/Bus/MessageRouter.h
+++ b/Engine/Source/Runtime/Messaging/Private/Bus/MessageRouter.h
@@ -116,13 +116,13 @@ public:
 
 	// Begin FRunnable interface
 
-	virtual bool Init( ) OVERRIDE;
+	virtual bool Init( ) override;
 
-	virtual uint32 Run( ) OVERRIDE;
+	virtual uint32 Run( ) override;
 
-	virtual void Stop( ) OVERRIDE;
+	virtual void Stop( ) override;
 
-	virtual void Exit( ) OVERRIDE;
+	virtual void Exit( ) override;
 
 	// End FRunnable interface
 
diff --git a/Engine/Source/Runtime/Messaging/Private/Bus/MessageSubscription.h b/Engine/Source/Runtime/Messaging/Private/Bus/MessageSubscription.h
--- a/Engine/Source/Runtime/Messaging/Private/Bus/MessageSubscription.h
+++ b/Engine/Source/Runtime/Messaging/Private/Bus/MessageSubscription.h
@@ -38,32 +38,32 @@ public:
 
 	// Begin IMessageSubscription interface
 
-	virtual void Disable( ) OVERRIDE
+	virtual void Disable( ) override
 	{
 		Enabled = false;
 	}
 
-	virtual void Enable( ) OVERRIDE
+	virtual void Enable( ) override
 	{
 		Enabled = true;
 	}
 
-	virtual FName GetMessageType( ) OVERRIDE
+	virtual FName GetMessageType( ) override
 	{
 		return MessageType;
 	}
 
-	virtual const FMessageScopeRange& GetScopeRange( ) OVERRIDE
+	virtual const FMessageScopeRange& GetScopeRange( ) override
 	{
 		return ScopeRange;
 	}
 
-	virtual const IReceiveMessagesWeakPtr& GetSubscriber( ) OVERRIDE
+	virtual const IReceiveMessagesWeakPtr& GetSubscriber( ) override
 	{
 		return Subscriber;
 	}
 
-	virtual bool IsEnabled( ) OVERRIDE
+	virtual bool IsEnabled( ) override
 	{
 		return Enabled;
 	}
diff --git a/Engine/Source/Runtime/Messaging/Private/Bus/MessageTracer.h b/Engine/Source/Runtime/Messaging/Private/Bus/MessageTracer.h
--- a/Engine/Source/Runtime/Messaging/Private/Bus/MessageTracer.h
+++ b/Engine/Source/Runtime/Messaging/Private/Bus/MessageTracer.h
@@ -225,12 +225,12 @@ public:
 
 	// Begin IMessageTracer interface
 
-	virtual void Break( ) OVERRIDE
+	virtual void Break( ) override
 	{
 		Breaking = true;
 	}
 
-	virtual void Continue( ) OVERRIDE
+	virtual void Continue( ) override
 	{
 		if (!Breaking)
 		{
@@ -241,48 +241,48 @@ public:
 		ContinueEvent->Trigger();
 	}
 
-	virtual int32 GetEndpoints( TArray<FMessageTracerEndpointInfoPtr>& OutEndpoints ) const OVERRIDE;
+	virtual int32 GetEndpoints( TArray<FMessageTracerEndpointInfoPtr>& OutEndpoints ) const override;
 
-	virtual int32 GetMessages( TArray<FMessageTracerMessageInfoPtr>& OutMessages ) const OVERRIDE;
+	virtual int32 GetMessages( TArray<FMessageTracerMessageInfoPtr>& OutMessages ) const override;
 
-	virtual int32 GetMessageTypes( TArray<FMessageTracerTypeInfoPtr>& OutTypes ) const OVERRIDE;
+	virtual int32 GetMessageTypes( TArray<FMessageTracerTypeInfoPtr>& OutTypes ) const override;
 
-	virtual bool HasMessages( ) const OVERRIDE
+	virtual bool HasMessages( ) const override
 	{
 		return (MessageInfos.Num() > 0);
 	}
 
-	virtual bool IsBreaking( ) const OVERRIDE
+	virtual bool IsBreaking( ) const override
 	{
 		return Breaking;
 	}
 
-	virtual bool IsRunning( ) const OVERRIDE
+	virtual bool IsRunning( ) const override
 	{
 		return Running;
 	}
 
-	virtual FMessageTracerMessageAdded& OnMessageAdded( ) OVERRIDE
+	virtual FMessageTracerMessageAdded& OnMessageAdded( ) override
 	{
 		return MessagesAddedDelegate;
 	}
 
-	virtual FSimpleMulticastDelegate& OnMessagesReset( ) OVERRIDE
+	virtual FSimpleMulticastDelegate& OnMessagesReset( ) override
 	{
 		return MessagesResetDelegate;
 	}
 
-	virtual FMessageTracerTypeAdded& OnTypeAdded( ) OVERRIDE
+	virtual FMessageTracerTypeAdded& OnTypeAdded( ) override
 	{
 		return TypeAddedDelegate;
 	}
 
-	virtual void Reset( ) OVERRIDE
+	virtual void Reset( ) override
 	{
 		ResetPending = true;
 	}
 
-	virtual void Start( ) OVERRIDE
+	virtual void Start( ) override
 	{
 		if (Running)
 		{
@@ -292,7 +292,7 @@ public:
 		Running = true;
 	}
 
-	virtual void Step( ) OVERRIDE
+	virtual void Step( ) override
 	{
 		if (!Breaking)
 		{
@@ -302,7 +302,7 @@ public:
 		ContinueEvent->Trigger();
 	}
 
-	virtual void Stop( ) OVERRIDE
+	virtual void Stop( ) override
 	{
 		if (!Running)
 		{
@@ -318,7 +318,7 @@ public:
 		}
 	}
 
-	virtual bool Tick( float DeltaTime ) OVERRIDE;
+	virtual bool Tick( float DeltaTime ) override;
 
 	// End IMessageTracer interface
 
diff --git a/Engine/Source/Runtime/Messaging/Private/Serialization/BsonMessageSerializer.h b/Engine/Source/Runtime/Messaging/Private/Serialization/BsonMessageSerializer.h
--- a/Engine/Source/Runtime/Messaging/Private/Serialization/BsonMessageSerializer.h
+++ b/Engine/Source/Runtime/Messaging/Private/Serialization/BsonMessageSerializer.h
@@ -15,13 +15,13 @@ class FBsonMessageSerializer
 {
 	// Begin ISerializeMessages interface
 
-	virtual bool DeserializeMessage( FArchive& Archive, IMutableMessageContextRef& OutContext ) OVERRIDE
+	virtual bool DeserializeMessage( FArchive& Archive, IMutableMessageContextRef& OutContext ) override
 	{
 		// not implemented yet
 		return false;
 	}
 
-	virtual bool SerializeMessage( const IMessageContextRef& Context, FArchive& Archive ) OVERRIDE
+	virtual bool SerializeMessage( const IMessageContextRef& Context, FArchive& Archive ) override
 	{
 		// not implemented yet
 		return false;
diff --git a/Engine/Source/Runtime/Messaging/Private/Serialization/JsonMessageSerializer.h b/Engine/Source/Runtime/Messaging/Private/Serialization/JsonMessageSerializer.h
--- a/Engine/Source/Runtime/Messaging/Private/Serialization/JsonMessageSerializer.h
+++ b/Engine/Source/Runtime/Messaging/Private/Serialization/JsonMessageSerializer.h
@@ -17,9 +17,9 @@ public:
 
 	// Begin ISerializeMessages interface
 
-	virtual bool DeserializeMessage( FArchive& Archive, IMutableMessageContextRef& OutContext ) OVERRIDE;
+	virtual bool DeserializeMessage( FArchive& Archive, IMutableMessageContextRef& OutContext ) override;
 
-	virtual bool SerializeMessage( const IMessageContextRef& Context, FArchive& Archive ) OVERRIDE;
+	virtual bool SerializeMessage( const IMessageContextRef& Context, FArchive& Archive ) override;
 
 	// End ISerializeMessages interface
 
diff --git a/Engine/Source/Runtime/Messaging/Private/Serialization/XmlMessageSerializer.h b/Engine/Source/Runtime/Messaging/Private/Serialization/XmlMessageSerializer.h
--- a/Engine/Source/Runtime/Messaging/Private/Serialization/XmlMessageSerializer.h
+++ b/Engine/Source/Runtime/Messaging/Private/Serialization/XmlMessageSerializer.h
@@ -15,13 +15,13 @@ class FXmlMessageSerializer
 {
 	// Begin ISerializeMessages interface
 
-	virtual bool DeserializeMessage( FArchive& Archive, IMutableMessageContextRef& OutContext ) OVERRIDE
+	virtual bool DeserializeMessage( FArchive& Archive, IMutableMessageContextRef& OutContext ) override
 	{
 		// not implemented yet
 		return false;
 	}
 
-	virtual bool SerializeMessage( const IMessageContextRef& Context, FArchive& Archive ) OVERRIDE
+	virtual bool SerializeMessage( const IMessageContextRef& Context, FArchive& Archive ) override
 	{
 		// not implemented yet
 		return false;
diff --git a/Engine/Source/Runtime/Messaging/Public/Common/FileMessageAttachment.h b/Engine/Source/Runtime/Messaging/Public/Common/FileMessageAttachment.h
--- a/Engine/Source/Runtime/Messaging/Public/Common/FileMessageAttachment.h
+++ b/Engine/Source/Runtime/Messaging/Public/Common/FileMessageAttachment.h
@@ -51,7 +51,7 @@ public:
 
 	// Begin IMessageAttachment interface
 
-	virtual FArchive* CreateReader( ) OVERRIDE
+	virtual FArchive* CreateReader( ) override
 	{
 		return IFileManager::Get().CreateFileReader(*Filename);
 	}
diff --git a/Engine/Source/Runtime/Messaging/Public/Common/MessageEndpoint.h b/Engine/Source/Runtime/Messaging/Public/Common/MessageEndpoint.h
--- a/Engine/Source/Runtime/Messaging/Public/Common/MessageEndpoint.h
+++ b/Engine/Source/Runtime/Messaging/Public/Common/MessageEndpoint.h
@@ -405,27 +405,27 @@ public:
 
 	// Begin IReceiveMessages interface
 
-	virtual FName GetDebugName( ) const OVERRIDE
+	virtual FName GetDebugName( ) const override
 	{
 		return Name;
 	}
 
-	virtual const FGuid& GetRecipientId( ) const OVERRIDE
+	virtual const FGuid& GetRecipientId( ) const override
 	{
 		return Id;
 	}
 
-	virtual ENamedThreads::Type GetRecipientThread( ) const OVERRIDE
+	virtual ENamedThreads::Type GetRecipientThread( ) const override
 	{
 		return RecipientThread;
 	}
 
-	virtual bool IsLocal( ) const OVERRIDE
+	virtual bool IsLocal( ) const override
 	{
 		return true;
 	}
 
-	virtual void ReceiveMessage( const IMessageContextRef& Context ) OVERRIDE
+	virtual void ReceiveMessage( const IMessageContextRef& Context ) override
 	{
 		if (!Enabled)
 		{
@@ -453,12 +453,12 @@ public:
 
 	// Begin ISendMessages interface
 
-	virtual FMessageAddress GetSenderAddress( ) OVERRIDE
+	virtual FMessageAddress GetSenderAddress( ) override
 	{
 		return Address;
 	}
 
-	virtual void NotifyMessageError( const IMessageContextRef& Context, const FString& Error ) OVERRIDE
+	virtual void NotifyMessageError( const IMessageContextRef& Context, const FString& Error ) override
 	{
 		ErrorDelegate.ExecuteIfBound(Context, Error);
 	}
diff --git a/Engine/Source/Runtime/Messaging/Public/Interfaces/IMessageHandler.h b/Engine/Source/Runtime/Messaging/Public/Interfaces/IMessageHandler.h
--- a/Engine/Source/Runtime/Messaging/Public/Interfaces/IMessageHandler.h
+++ b/Engine/Source/Runtime/Messaging/Public/Interfaces/IMessageHandler.h
@@ -86,12 +86,12 @@ public:
 
 	// Begin IMessageHandler interface
 	
-	virtual const FName GetHandledMessageType( ) const OVERRIDE
+	virtual const FName GetHandledMessageType( ) const override
 	{
 		return MessageType::StaticStruct()->GetFName();
 	}
 
-	virtual void HandleMessage( const IMessageContextRef& Context ) OVERRIDE
+	virtual void HandleMessage( const IMessageContextRef& Context ) override
 	{
 		(Handler->*HandlerFunc)(*static_cast<const MessageType*>(Context->GetMessage()), Context);
 	}
diff --git a/Engine/Source/Runtime/Networking/Private/NetworkingModule.h b/Engine/Source/Runtime/Networking/Private/NetworkingModule.h
--- a/Engine/Source/Runtime/Networking/Private/NetworkingModule.h
+++ b/Engine/Source/Runtime/Networking/Private/NetworkingModule.h
@@ -15,9 +15,9 @@ class FNetworkingModule
 {
 public:
 
-	virtual void StartupModule( ) OVERRIDE;
+	virtual void StartupModule( ) override;
 
-	virtual void ShutdownModule( ) OVERRIDE;
+	virtual void ShutdownModule( ) override;
 
-	virtual bool SupportsDynamicReloading( ) OVERRIDE;
+	virtual bool SupportsDynamicReloading( ) override;
 };
\ No newline at end of file
diff --git a/Engine/Source/Runtime/Networking/Public/Common/TcpListener.h b/Engine/Source/Runtime/Networking/Public/Common/TcpListener.h
--- a/Engine/Source/Runtime/Networking/Public/Common/TcpListener.h
+++ b/Engine/Source/Runtime/Networking/Public/Common/TcpListener.h
@@ -135,7 +135,7 @@ public:
 
 	// Begin FRunnable interface
 
-	virtual bool Init( ) OVERRIDE
+	virtual bool Init( ) override
 	{
 		if (Socket == NULL)
 		{
@@ -151,7 +151,7 @@ public:
 		return (Socket != NULL);
 	}
 
-	virtual uint32 Run( ) OVERRIDE
+	virtual uint32 Run( ) override
 	{
 		TSharedRef<FInternetAddr> RemoteAddress = ISocketSubsystem::Get(PLATFORM_SOCKETSUBSYSTEM)->CreateInternetAddr();
 
@@ -187,12 +187,12 @@ public:
 		return 0;
 	}
 
-	virtual void Stop( ) OVERRIDE
+	virtual void Stop( ) override
 	{
 		Stopping = true;
 	}
 
-	virtual void Exit( ) OVERRIDE { }
+	virtual void Exit( ) override { }
 
 	// End FRunnable interface
 
diff --git a/Engine/Source/Runtime/Networking/Public/Common/UdpSocketReceiver.h b/Engine/Source/Runtime/Networking/Public/Common/UdpSocketReceiver.h
--- a/Engine/Source/Runtime/Networking/Public/Common/UdpSocketReceiver.h
+++ b/Engine/Source/Runtime/Networking/Public/Common/UdpSocketReceiver.h
@@ -76,12 +76,12 @@ public:
 
 public:
 
-	virtual bool Init( ) OVERRIDE
+	virtual bool Init( ) override
 	{
 		return true;
 	}
 
-	virtual uint32 Run( ) OVERRIDE
+	virtual uint32 Run( ) override
 	{
 		TSharedRef<FInternetAddr> Sender = SocketSubsystem->CreateInternetAddr();
 
@@ -109,12 +109,12 @@ public:
 		return 0;
 	}
 
-	virtual void Stop( ) OVERRIDE
+	virtual void Stop( ) override
 	{
 		Stopping = true;
 	}
 
-	virtual void Exit( ) OVERRIDE { }
+	virtual void Exit( ) override { }
 
 
 private:
diff --git a/Engine/Source/Runtime/Networking/Public/Common/UdpSocketSender.h b/Engine/Source/Runtime/Networking/Public/Common/UdpSocketSender.h
--- a/Engine/Source/Runtime/Networking/Public/Common/UdpSocketSender.h
+++ b/Engine/Source/Runtime/Networking/Public/Common/UdpSocketSender.h
@@ -137,12 +137,12 @@ public:
 
 public:
 
-	virtual bool Init( ) OVERRIDE
+	virtual bool Init( ) override
 	{
 		return true;
 	}
 
-	virtual uint32 Run( ) OVERRIDE
+	virtual uint32 Run( ) override
 	{
 		while (!Stopping)
 		{
@@ -171,14 +171,14 @@ public:
 		return 0;
 	}
 
-	virtual void Stop( ) OVERRIDE
+	virtual void Stop( ) override
 	{
 		Stopping = true;
 
 		WorkEvent->Trigger();
 	}
 
-	virtual void Exit( ) OVERRIDE { }
+	virtual void Exit( ) override { }
 
 private:
 
