Messaging: Added more documentation, minor fixes

diff --git a/Engine/Source/Runtime/Messaging/Public/Interfaces/IMessageAttachment.h b/Engine/Source/Runtime/Messaging/Public/Interfaces/IMessageAttachment.h
--- a/Engine/Source/Runtime/Messaging/Public/Interfaces/IMessageAttachment.h
+++ b/Engine/Source/Runtime/Messaging/Public/Interfaces/IMessageAttachment.h
@@ -16,6 +16,18 @@ typedef TSharedRef<class IMessageAttachment, ESPMode::ThreadSafe> IMessageAttach
 
 /**
  * Interface for message attachments.
+ *
+ * Message attachments are optional bulk data objects that can be attached to a message, such as
+ * files or memory buffers. Transferring attachments within the same process is very efficient and
+ * amounts to copying a pointer to the data. On the other hand, sending message attachments to
+ * applications running in other processes or on other computers may - depending on their size -
+ * take a considerable amount of time, because they need to be serialized and transferred. For this
+ * reason, attachments are transferred separately from messages in such cases.
+ *
+ * NOTE: Message attachments are not fully implemented yet. In particular, transferring attachments
+ * over a network does not work yet. This API may change in the near future, so please use with care.
+ *
+ * @see IMessageContext
  */
 class IMessageAttachment
 {
diff --git a/Engine/Source/Runtime/Messaging/Public/Interfaces/IMessageBridge.h b/Engine/Source/Runtime/Messaging/Public/Interfaces/IMessageBridge.h
--- a/Engine/Source/Runtime/Messaging/Public/Interfaces/IMessageBridge.h
+++ b/Engine/Source/Runtime/Messaging/Public/Interfaces/IMessageBridge.h
@@ -16,6 +16,17 @@ typedef TSharedRef<class IMessageBridge, ESPMode::ThreadSafe> IMessageBridgeRef;
 
 /**
  * Interface for message bridges.
+ *
+ * A message bridge connects a message bus with another messaging system. It is really just a regular
+ * message endpoint connected to a message bus that translates sent and received messages between the bus
+ * and some underlying transport technology. The transport technology is usually implemented in the form of a 
+ * Message Transport Plug-in, such as the UdpMessaging plug-in that ships with Unreal Engine.
+ *
+ * The most common use case for message bridges is to connect two Unreal Engine message buses running in
+ * separate processes or en different computers. Another common use case is to connect an Unreal Engine
+ * message bus to an entirely different messaging system that is not based on Unreal Engine.
+ *
+ * Message bridge instances can be created with the @see IMessagingModule.CreateBridge method.
  */
 class IMessageBridge
 {
diff --git a/Engine/Source/Runtime/Messaging/Public/Interfaces/IMessageBus.h b/Engine/Source/Runtime/Messaging/Public/Interfaces/IMessageBus.h
--- a/Engine/Source/Runtime/Messaging/Public/Interfaces/IMessageBus.h
+++ b/Engine/Source/Runtime/Messaging/Public/Interfaces/IMessageBus.h
@@ -25,6 +25,82 @@ DECLARE_MULTICAST_DELEGATE(FOnMessageBusShutdown);
 
 /**
  * Interface for message buses.
+ *
+ * A message bus is the main logical component to facilitate communication between (possibly distributed) parts of an application
+ * using Message Passing as its underlying architectural pattern. It allows registered sender and recipient objects, subsequently
+ * referred to as Message Endpoints, to exchange structured data in the form of user defined messages[1].
+ *
+ * Depending on their usage, messages are classified into a number of types, such as commands, events and documents. In Unreal Engine,
+ * all these messages are modeled with regular built-in or user defined UStructs that may either be empty or contain data in the
+ * form of UProperty fields[2]. Before being dispatched, messages are internally wrapped into so called Message Context objects
+ * (@see IMessageContext) that contain additional information about a message, such as when it was sent and the sender and recipients.
+ *
+ * The sending and receiving of messages is not limited to message endpoints within the same thread or process, but may be extended
+ * to other applications running on the same computer or on other computers connected to a network with so called Message Transport
+ * Plug-ins, such as the UdpMessaging plug-in that ships with Unreal Engine. The main goal of a message bus is to hide the technical
+ * details of the underlying transport mechanisms, so that users can focus on implementing their distributed applications rather than
+ * worrying about how data gets from one endpoint to another. A message bus makes it look as if all senders and recipients are located
+ * within the same process, regardless of whether that is actually the case or not. It is possible to restrict the reach of messages
+ * using so called Message Scopes (@see EMessageScope).
+ *
+ * All message recipients (objects implementing the @see IReceiveMessages interface) must be registered with the message bus through the
+ * @see IMessageBus.Register method. Before a recipient is destroyed it should unregister itself using the @see IMessageBus.Unregister
+ * method. Message senders (objects implementing the @see ISendMessages interface) do not register with the bus, but instead pass a
+ * reference to themselves each time they send a message. The IReceiveMessages and ISendMessages both are very low-level interfaces
+ * into the messaging system. Most users will prefer to use instances of the @see FMessageEndpoint class instead, which provides a much
+ * more convenient way of sending and receiving messages.
+ *
+ * Message buses in Unreal Engine support the following two common messaging patterns: Request-Reply and Publish-Subscribe. Please note
+ * that higher level concepts, such as Remote Procedure Call (RPC) are not part of the messaging system, but we may provide them as
+ * separate features in the future.
+ *
+ * In the Request-Reply pattern a message is sent to one or more particular message recipients using the IMessageBus.Send method.
+ * Message recipients implement the @see IMessageRecipient interface and are uniquely identified by their addresses (@see FMessageAddress).
+ * After a message is received, the recipients may then reply with another message using the same IMessageBus.Send method. Alternatively,
+ * a previously received message may be forwarded to another recipient using the @see IMessageBus.Forward method. This pattern is useful
+ * when message recipients already know about each other and wish to communicate directly, i.e. to exchange commands or events.
+ *
+ * In the Publish-Subscribe pattern a message is sent to all message recipients on the bus using the @see IMessageBus.Publish method.
+ * Only recipients that previously subscribed to the type of the sent message using the @IMessageBus.Subscribe method will actually receive
+ * the message. All other recipients will not receive the message. After a published message is received, the recipients may respond with
+ * another message, either directly to the message sender using IMessageBus.Send method, or by publishing another message using the
+ * the IMessageBus.Publish method. This pattern is useful for discovering recipients on the message bus and to recipients that are unknown.
+ *
+ * Most applications will often use a combination of both Request-Reply and Publish-Subscribe. A typical implementation of distributed
+ * applications involves service providers and service consumers (a service can be any useful functionality provided by a system). The
+ * service consumers will often publish a special message in order to discover service providers. Service providers are subscribed to these
+ * special messages and will respond with another message that contains information about the provider and that is sent back directly to the
+ * consumer. The consumer is now aware of the service provider's existence and its message address, and it can then request services by
+ * directly sending all subsequent messages directly to it.
+ *
+ * The dispatching of sent and published messages to the correct recipients is accomplished by the Message Router, which is an internal
+ * component of the message bus. The message router maintains an address book of known message endpoints that allows it to determine
+ * the destination of messages. If a message cannot be delivered, it is forwarded to a so called Dead Letter Channel, which is a queue
+ * of messages that can be inspected for debugging purposes[3].
+ *
+ * It is possible to intercept messages before they are being routed to recipients by registering so called Message Interceptors (objects
+ * implementing the IInterceptMessage interface) with the message bus. This feature enables advanced use cases that require inspection or
+ * manipulation of messages contents, such as message filtering and enriching, splitting and aggregating, resquencing or authentication[4].
+ *
+ * The messaging system also provides a facility for debugging the system itself through the so called Message Tracer (an internal object
+ * implementing the @see IMessageTracer interface) that can be accessed with the IMessageBus.GetTracer() method. The message tracer is
+ * currently used in the Messaging Debugger - a visual debugging tool for the Messaging System in Unreal Frontend and the Unreal Editor.
+ *
+ * Notes:
+ *
+ * [1] The messaging implementation in Unreal Engine is more or less closely following the terminology in the Enterprise Integration
+ * Patterns book at http://www.eaipatterns.com/eaipatterns.html. A good introduction to messaging concepts can also be found on MSDN
+ * at http://msdn.microsoft.com/en-us/library/ff647328.aspx
+ *
+ * [2] It is currently also required for such structures to implement a specific TStructOpsTypeTraits type trait or otherwise they
+ * will not be recognized as legal message types. Search for 'WithMessagingHandling' in the code base to find existing examples. We
+ * hope to remove the need for these type traits soon.
+ *
+ * [3] Dead letter channels are not implemented yet.
+ *
+ * [4] There are currently no built-in mechanisms for authentication and authorization, but they are on the to-do list.
+ *
+ * @see IReceiveMessages, ISendMessages, FMessageAddress, FMessageEndpoint
  */
 class IMessageBus
 {
@@ -40,6 +116,7 @@ public:
 	 * @param NewScope The forwarding scope.
 	 * @param Delay The deferral time.
 	 * @param Forwarder The sender that forwards the message.
+	 * @see Publish, Send
 	 */
 	virtual void Forward( const IMessageContextRef& Context, const TArray<FMessageAddress>& Recipients, EMessageScope::Type NewScope, const FTimespan& Delay, const ISendMessagesRef& Forwarder ) = 0;
 
@@ -55,6 +132,7 @@ public:
 	 *
 	 * @param Interceptor The interceptor.
 	 * @param MessageType The type of messages to intercept.
+	 * @see Unintercept
 	 */
 	virtual void Intercept( const IInterceptMessagesRef& Interceptor, const FName& MessageType ) = 0;
 
@@ -69,6 +147,7 @@ public:
 	 * @param Delay The delay after which to send the message.
 	 * @param Expiration The time at which the message expires.
 	 * @param Publisher The message publisher.
+	 * @see Forward, Send
 	 */
 	virtual void Publish( void* Message, UScriptStruct* TypeInfo, EMessageScope::Type Scope, const FTimespan& Delay, const FDateTime& Expiration, const ISendMessagesRef& Publisher ) = 0;
 
@@ -77,6 +156,7 @@ public:
 	 *
 	 * @param Address The address of the recipient to register.
 	 * @param Recipient The message recipient.
+	 * @see Unregister
 	 */
 	virtual void Register( const FMessageAddress& Address, const IReceiveMessagesRef& Recipient) = 0;
 
@@ -90,6 +170,7 @@ public:
 	 * @param Delay The delay after which to send the message.
 	 * @param Expiration The time at which the message expires.
 	 * @param Sender The message sender.
+	 * @see Forward, Publish
 	 */
 	virtual void Send( void* Message, UScriptStruct* TypeInfo, const IMessageAttachmentPtr& Attachment, const TArray<FMessageAddress>& Recipients, const FTimespan& Delay, const FDateTime& Expiration, const ISendMessagesRef& Sender ) = 0;
 
@@ -119,6 +200,7 @@ public:
 	 *
 	 * @param Interceptor The interceptor to remove.
 	 * @param MessageType The type of messages to stop intercepting.
+	 * @see Intercept
 	 */
 	virtual void Unintercept( const IInterceptMessagesRef& Interceptor, const FName& MessageType ) = 0;
 
@@ -126,6 +208,7 @@ public:
 	 * Unregisters a message recipient from the message bus.
 	 *
 	 * @param Address The address of the recipient to unregister.
+	 * @see Register
 	 */
 	virtual void Unregister( const FMessageAddress& Address ) = 0;
 
diff --git a/Engine/Source/Runtime/Messaging/Public/Interfaces/IMessageContext.h b/Engine/Source/Runtime/Messaging/Public/Interfaces/IMessageContext.h
--- a/Engine/Source/Runtime/Messaging/Public/Interfaces/IMessageContext.h
+++ b/Engine/Source/Runtime/Messaging/Public/Interfaces/IMessageContext.h
@@ -63,6 +63,20 @@ typedef TRangeBound<EMessageScope::Type> FMessageScopeRangeBound;
 
 /**
  * Interface for message contexts.
+ *
+ * A message context is created internally by the message bus to store additional information about a
+ * message needed to describe and correctly route the message to recipients. The two most common aspects
+ * of message contexts used by recipients are the sender address and message attachments.
+ *
+ * The sender address (@see GetSender) is often needed to send a reply message to a message sender, i.e.
+ * in response to a published message. The message attachment (@see GetAttchment) is an optional bundle
+ * of binary bulk data that is transferred independently from the message itself and allows for transferring
+ * larger amounts of data that would otherwise clog up the messaging system.
+ *
+ * In case a message was forwarded by another endpoint, the context of the original sender can be accessed
+ * using the @see GetOriginalContext method.
+ *
+ * @see IMessageAttachment
  */
 class IMessageContext
 {
@@ -71,7 +85,7 @@ public:
 	/**
 	 * Gets the message attachment, if present.
 	 *
-	 * @return A pointer to the message attachment, or NULL if no attachment is present.
+	 * @return A pointer to the message attachment, or nullptr if no attachment is present.
 	 */
 	virtual IMessageAttachmentPtr GetAttachment( ) const = 0;
 
diff --git a/Engine/Source/Runtime/Messaging/Public/Interfaces/IMessageHandler.h b/Engine/Source/Runtime/Messaging/Public/Interfaces/IMessageHandler.h
--- a/Engine/Source/Runtime/Messaging/Public/Interfaces/IMessageHandler.h
+++ b/Engine/Source/Runtime/Messaging/Public/Interfaces/IMessageHandler.h
@@ -79,7 +79,7 @@ public:
 		: Handler(InHandler)
 		, HandlerFunc(InHandlerFunc)
 	{
-		check(InHandler != NULL);
+		check(InHandler != nullptr);
 	}
 
 public:
diff --git a/Engine/Source/Runtime/Messaging/Public/Interfaces/IMessageSubscription.h b/Engine/Source/Runtime/Messaging/Public/Interfaces/IMessageSubscription.h
--- a/Engine/Source/Runtime/Messaging/Public/Interfaces/IMessageSubscription.h
+++ b/Engine/Source/Runtime/Messaging/Public/Interfaces/IMessageSubscription.h
@@ -16,6 +16,14 @@ typedef TSharedRef<class IMessageSubscription, ESPMode::ThreadSafe> IMessageSubs
 
 /**
  * Interface for message subscriptions.
+ *
+ * This interface provides access to a message subscription that was previously created with the
+ * @see IMessageBus.Subscribe method. It can be used to query the subscription's details using the
+ * various getters and modify its enabled state using the @see IMessageSubscription.Enable and
+ * IMessageSubscription.Disable methods.
+ *
+ * A subscription that is disabled will cause a message endpoint to stop receiving messages for
+ * the subscribed message type. However, the subscription is merely disabled and not removed.
  */
 class IMessageSubscription
 {
diff --git a/Engine/Source/Runtime/Messaging/Public/Interfaces/IMessageTracer.h b/Engine/Source/Runtime/Messaging/Public/Interfaces/IMessageTracer.h
--- a/Engine/Source/Runtime/Messaging/Public/Interfaces/IMessageTracer.h
+++ b/Engine/Source/Runtime/Messaging/Public/Interfaces/IMessageTracer.h
@@ -178,6 +178,10 @@ DECLARE_MULTICAST_DELEGATE_OneParam(FMessageTracerTypeAdded, FMessageTracerTypeI
 
 /**
  * Interface for message tracers.
+ *
+ * This interface provides access to the message tracer that is built into a message bus. Message tracers are
+ * able to monitor and record all activities on a message bus. This includes the sending and receiving messages,
+ * the registering and unregistering of message endpoints and interceptors and adding and removal of subscriptions.
  */
 class IMessageTracer
 {
diff --git a/Engine/Source/Runtime/Messaging/Public/Interfaces/IMessageTracerBreakpoint.h b/Engine/Source/Runtime/Messaging/Public/Interfaces/IMessageTracerBreakpoint.h
--- a/Engine/Source/Runtime/Messaging/Public/Interfaces/IMessageTracerBreakpoint.h
+++ b/Engine/Source/Runtime/Messaging/Public/Interfaces/IMessageTracerBreakpoint.h
@@ -16,6 +16,8 @@ typedef TSharedRef<class IMessageTracerBreakpoint, ESPMode::ThreadSafe> IMessage
 
 /**
  * Interface for message tracer breakpoints.
+ *
+ * @see IMessageTracer
  */
 class IMessageTracerBreakpoint
 {
diff --git a/Engine/Source/Runtime/Messaging/Public/Interfaces/IMessagingModule.h b/Engine/Source/Runtime/Messaging/Public/Interfaces/IMessagingModule.h
--- a/Engine/Source/Runtime/Messaging/Public/Interfaces/IMessagingModule.h
+++ b/Engine/Source/Runtime/Messaging/Public/Interfaces/IMessagingModule.h
@@ -9,6 +9,8 @@
 
 /**
  * Interface for messaging modules.
+ *
+ * @see IMessageBridge, IMessageBus
  */
 class IMessagingModule
 	: public IModuleInterface
diff --git a/Engine/Source/Runtime/Messaging/Public/Interfaces/IReceiveMessages.h b/Engine/Source/Runtime/Messaging/Public/Interfaces/IReceiveMessages.h
--- a/Engine/Source/Runtime/Messaging/Public/Interfaces/IReceiveMessages.h
+++ b/Engine/Source/Runtime/Messaging/Public/Interfaces/IReceiveMessages.h
@@ -20,10 +20,14 @@ typedef TWeakPtr<class IReceiveMessages, ESPMode::ThreadSafe> IReceiveMessagesWe
 /**
  * Interface for message recipients.
  *
- * This is a low-level interface for message recipients that are subscribed to a message bus.
- * For more convenient messaging functionality it is recommended to use message endpoints instead.
+ * Classes that implement this interface are able to receive messages from a message bus. A message recipient will receive
+ * a call to its ReceiveMessage() method for each message that was sent directly to it (@see IMessageBus.Send) and for each
+ * published message that it subscribed to (@see IMessageBus.Publish).
  *
- * @see IMessageEndpoint
+ * This interface provides a rather low-level mechanism for receiving messages. Instead of implementing it, Most users will
+ * want to use an instance of @see FMessageEndpoint, which provides a much more convenient way of sending and receiving messages.
+ *
+ * @see FMessageEndpoint, IMessageBus, ISendMessages
  */
 class IReceiveMessages
 {
@@ -33,7 +37,7 @@ public:
 	 * Gets the recipient's name (for debugging purposes).
 	 *
 	 * @return The debug name.
-	 * @see GetRecipientId
+	 * @see GetRecipientId, GetRecipientThread
 	 */
 	virtual FName GetDebugName( ) const = 0;
 
@@ -41,7 +45,7 @@ public:
 	 * Gets the recipient's unique identifier (for debugging purposes).
 	 *
 	 * @return The recipient's identifier.
-	 * @see GetRecipientName
+	 * @see GetDebugName, GetRecipientThread
 	 */
 	virtual const FGuid& GetRecipientId( ) const = 0;
 
@@ -51,6 +55,7 @@ public:
 	 * If the recipient's ReceiveMessage() is thread-safe, return ThreadAny for best performance.
 	 *
 	 * @return Name of the receiving thread.
+	 * @see GetDebugName, GetRecipientId
 	 */
 	virtual ENamedThreads::Type GetRecipientThread( ) const = 0;
 
@@ -61,6 +66,7 @@ public:
 	 * other processes on the same machine or on remote machines are considered remote.
 	 *
 	 * @return true if this recipient is local, false otherwise.
+	 * @see IsRemote
 	 */
 	virtual bool IsLocal( ) const = 0;
 
@@ -80,6 +86,7 @@ public:
 	 * other processes on the same machine or on remote machines are considered remote.
 	 *
 	 * @return true if this recipient is remote, false otherwise.
+	 * @see IsLocal
 	 */
 	bool IsRemote( ) const
 	{
diff --git a/Engine/Source/Runtime/Messaging/Public/Interfaces/ISendMessages.h b/Engine/Source/Runtime/Messaging/Public/Interfaces/ISendMessages.h
--- a/Engine/Source/Runtime/Messaging/Public/Interfaces/ISendMessages.h
+++ b/Engine/Source/Runtime/Messaging/Public/Interfaces/ISendMessages.h
@@ -14,6 +14,22 @@ typedef TSharedPtr<class ISendMessages, ESPMode::ThreadSafe> ISendMessagesPtr;
 typedef TSharedRef<class ISendMessages, ESPMode::ThreadSafe> ISendMessagesRef;
 
 
+/**
+ * Interface for message senders.
+ *
+ * Classes that implement this interface are able to send messages on a message bus. Each message sender must be uniquely
+ * identifiable with a message address that is returned by the @see ISendMessages.GetSenderAddress method. It is recommended
+ * that implementors of this interface generate a GUID for each instance, which must not change throughout the lifetime of
+ * the instance.
+ *
+ * The sending of messages is accomplished with the @see IMessageBus.Forward, IMessageBus.Publish and IMessageBus.Send methods.
+ * In case an error occurs during the sending of a message, the @see ISendMessages.NotifyMessageError method will be called.
+ *
+ * This interface provides a rather low-level mechanism for receiving messages. Instead of implementing it, Most users will
+ * want to use an instance of @see FMessageEndpoint, which provides a much more convenient way of sending and receiving messages.
+ *
+ * @see FMessageEndpoint, IMessageBus, IReceiveMessages
+ */
 class ISendMessages
 {
 public:
diff --git a/Engine/Source/Runtime/Messaging/Public/Interfaces/ITransportMessages.h b/Engine/Source/Runtime/Messaging/Public/Interfaces/ITransportMessages.h
--- a/Engine/Source/Runtime/Messaging/Public/Interfaces/ITransportMessages.h
+++ b/Engine/Source/Runtime/Messaging/Public/Interfaces/ITransportMessages.h
@@ -45,7 +45,7 @@ DECLARE_DELEGATE_OneParam(FOnMessageTransportNodeLost, const FGuid&);
  * Interface for message transport technologies.
  *
  * Licensees can implement this interface to add support for custom message transport
- * technologies that are not supported, i.e. custom network protocols or APIs.
+ * technologies that are not supported out of the box, i.e. custom network protocols or APIs.
  */
 class ITransportMessages
 {
@@ -108,5 +108,8 @@ public:
 
 protected:
 
-	ITransportMessages( ) { }
+	/**
+	 * Virtual constructor.
+	 */
+	~ITransportMessages( ) { }
 };
