renamed NULL to nullptr

diff --git a/Engine/Source/Runtime/Messaging/Public/Common/MessageEndpoint.h b/Engine/Source/Runtime/Messaging/Public/Common/MessageEndpoint.h
--- a/Engine/Source/Runtime/Messaging/Public/Common/MessageEndpoint.h
+++ b/Engine/Source/Runtime/Messaging/Public/Common/MessageEndpoint.h
@@ -583,7 +583,7 @@ public:
 	template<typename MessageType>
 	void Send( MessageType* Message, const FMessageAddress& Recipient )
 	{
-		Send(Message, MessageType::StaticStruct(), NULL, TArrayBuilder<FMessageAddress>().Add(Recipient), FTimespan::Zero(), FDateTime::MaxValue());
+		Send(Message, MessageType::StaticStruct(), nullptr, TArrayBuilder<FMessageAddress>().Add(Recipient), FTimespan::Zero(), FDateTime::MaxValue());
 	}
 
 	/**
@@ -597,7 +597,7 @@ public:
 	template<typename MessageType>
 	void Send( MessageType* Message, const FMessageAddress& Recipient, const FTimespan& Delay )
 	{
-		Send(Message, MessageType::StaticStruct(), NULL, TArrayBuilder<FMessageAddress>().Add(Recipient), Delay, FDateTime::MaxValue());
+		Send(Message, MessageType::StaticStruct(), nullptr, TArrayBuilder<FMessageAddress>().Add(Recipient), Delay, FDateTime::MaxValue());
 	}
 
 	/**
@@ -612,7 +612,7 @@ public:
 	template<typename MessageType>
 	void Send( MessageType* Message, const FMessageAddress& Recipient, const FTimespan& Delay, const FDateTime& Expiration )
 	{
-		Send(Message, MessageType::StaticStruct(), NULL, TArrayBuilder<FMessageAddress>().Add(Recipient), Delay, Expiration);
+		Send(Message, MessageType::StaticStruct(), nullptr, TArrayBuilder<FMessageAddress>().Add(Recipient), Delay, Expiration);
 	}
 
 	/**
@@ -655,7 +655,7 @@ public:
 	template<typename MessageType>
 	void Send( MessageType* Message, const TArray<FMessageAddress>& Recipients )
 	{
-		Send(Message, MessageType::StaticStruct(), NULL, Recipients, FTimespan::Zero(), FDateTime::MaxValue());
+		Send(Message, MessageType::StaticStruct(), nullptr, Recipients, FTimespan::Zero(), FDateTime::MaxValue());
 	}
 
 	/**
@@ -669,7 +669,7 @@ public:
 	template<typename MessageType>
 	void Send( MessageType* Message, const TArray<FMessageAddress>& Recipients, const FTimespan& Delay )
 	{
-		Send(Message, MessageType::StaticStruct(), NULL, Recipients, Delay, FDateTime::MaxValue());
+		Send(Message, MessageType::StaticStruct(), nullptr, Recipients, Delay, FDateTime::MaxValue());
 	}
 
 	/**
diff --git a/Engine/Source/Runtime/Messaging/Public/Common/MessageEndpointBuilder.h b/Engine/Source/Runtime/Messaging/Public/Common/MessageEndpointBuilder.h
--- a/Engine/Source/Runtime/Messaging/Public/Common/MessageEndpointBuilder.h
+++ b/Engine/Source/Runtime/Messaging/Public/Common/MessageEndpointBuilder.h
@@ -152,7 +152,7 @@ public:
 	/**
 	 * Builds the message endpoint as configured.
 	 *
-	 * @return A new message endpoint, or NULL if it couldn't be built.
+	 * @return A new message endpoint, or nullptr if it couldn't be built.
 	 */
 	FMessageEndpointPtr Build( )
 	{
diff --git a/Engine/Source/Runtime/Messaging/Public/Interfaces/IMessageBus.h b/Engine/Source/Runtime/Messaging/Public/Interfaces/IMessageBus.h
--- a/Engine/Source/Runtime/Messaging/Public/Interfaces/IMessageBus.h
+++ b/Engine/Source/Runtime/Messaging/Public/Interfaces/IMessageBus.h
@@ -190,7 +190,7 @@ public:
 	 * @param Subscriber The subscriber wishing to receive the messages.
 	 * @param MessageType The type of messages to subscribe to (NAME_All = subscribe to all message types).
 	 * @param ScopeRange The range of message scopes to include in the subscription.
-	 * @return The added subscription, or NULL if the subscription failed.
+	 * @return The added subscription, or nullptr if the subscription failed.
 	 * @see Unsubscribe
 	 */
 	virtual IMessageSubscriptionPtr Subscribe( const IReceiveMessagesRef& Subscriber, const FName& MessageType, const FMessageScopeRange& ScopeRange ) = 0;
diff --git a/Engine/Source/Runtime/Messaging/Public/Interfaces/IMessagingModule.h b/Engine/Source/Runtime/Messaging/Public/Interfaces/IMessagingModule.h
--- a/Engine/Source/Runtime/Messaging/Public/Interfaces/IMessagingModule.h
+++ b/Engine/Source/Runtime/Messaging/Public/Interfaces/IMessagingModule.h
@@ -27,7 +27,7 @@ public:
 	 * @param Bus The message bus to attach the bridge to.
 	 * @param Serializer The message serializer to use.
 	 * @param Transport The message transport technology to use.
-	 * @return The new message bridge, or NULL if the bridge couldn't be created.
+	 * @return The new message bridge, or nullptr if the bridge couldn't be created.
 	 * @see CreateBus
 	 */
 	virtual IMessageBridgePtr CreateBridge( const FMessageAddress& Address, const IMessageBusRef& Bus, const ISerializeMessagesRef& Serializer, const ITransportMessagesRef& Transport ) = 0;
@@ -36,7 +36,7 @@ public:
 	 * Creates a new message bus.
 	 *
 	 * @param RecipientAuthorizer An optional recipient authorizer.
-	 * @return The new message bus, or NULL if the bus couldn't be created.
+	 * @return The new message bus, or nullptr if the bus couldn't be created.
 	 * @see CreateBridge
 	 */
 	virtual IMessageBusPtr CreateBus( const IAuthorizeMessageRecipientsPtr& RecipientAuthorizer ) = 0;
