Adding missing MessageEndpoint.cpp

diff --git a/Engine/Source/Runtime/Messaging/Private/MessageEndpoint.cpp b/Engine/Source/Runtime/Messaging/Private/MessageEndpoint.cpp
new file mode 100644
--- /dev/null
+++ b/Engine/Source/Runtime/Messaging/Private/MessageEndpoint.cpp
@@ -0,0 +1,100 @@
+#include "MessagingPrivatePCH.h"
+#include "MessageEndpoint.h"
+
+MESSAGING_API FMessageEndpoint::FMessageEndpoint(const FName& InName, const IMessageBusRef& InBus, const TArray<IMessageHandlerPtr>& InHandlers)
+	: Address(FGuid::NewGuid())
+	, BusPtr(InBus)
+	, Enabled(true)
+	, Handlers(InHandlers)
+	, Id(FGuid::NewGuid())
+	, InboxEnabled(false)
+	, Name(InName)
+{ }
+
+MESSAGING_API FMessageEndpoint::~FMessageEndpoint()
+{
+	IMessageBusPtr Bus = BusPtr.Pin();
+
+	if (Bus.IsValid())
+	{
+		Bus->Unregister(Address);
+	}
+}
+
+MESSAGING_API void FMessageEndpoint::Defer(const IMessageContextRef& Context, const FTimespan& Delay)
+{
+	IMessageBusPtr Bus = GetBusIfEnabled();
+
+	if (Bus.IsValid())
+	{
+		Bus->Forward(Context, TArrayBuilder<FMessageAddress>().Add(Address), Context->GetScope(), Delay, AsShared());
+	}
+}
+
+MESSAGING_API void FMessageEndpoint::Forward(const IMessageContextRef& Context, const TArray<FMessageAddress>& Recipients, EMessageScope::Type ForwardingScope, const FTimespan& Delay)
+{
+	IMessageBusPtr Bus = GetBusIfEnabled();
+
+	if (Bus.IsValid())
+	{
+		Bus->Forward(Context, Recipients, ForwardingScope, Delay, AsShared());
+	}
+}
+
+MESSAGING_API void FMessageEndpoint::Publish(void* Message, UScriptStruct* TypeInfo, EMessageScope::Type Scope, const FTimespan& Delay, const FDateTime& Expiration)
+{
+	IMessageBusPtr Bus = GetBusIfEnabled();
+
+	if (Bus.IsValid())
+	{
+		Bus->Publish(Message, TypeInfo, Scope, Delay, Expiration, AsShared());
+	}
+}
+
+MESSAGING_API void FMessageEndpoint::Send(void* Message, UScriptStruct* TypeInfo, const IMessageAttachmentPtr& Attachment, const TArray<FMessageAddress>& Recipients, const FTimespan& Delay, const FDateTime& Expiration)
+{
+	IMessageBusPtr Bus = GetBusIfEnabled();
+
+	if (Bus.IsValid())
+	{
+		Bus->Send(Message, TypeInfo, Attachment, Recipients, Delay, Expiration, AsShared());
+	}
+}
+
+MESSAGING_API void FMessageEndpoint::Subscribe(const FName& MessageType, const FMessageScopeRange& ScopeRange)
+{
+	IMessageBusPtr Bus = GetBusIfEnabled();
+
+	if (Bus.IsValid())
+	{
+		Bus->Subscribe(AsShared(), MessageType, ScopeRange);
+	}
+}
+
+MESSAGING_API void FMessageEndpoint::Unsubscribe(const FName& TopicPattern)
+{
+	IMessageBusPtr Bus = GetBusIfEnabled();
+
+	if (Bus.IsValid())
+	{
+		Bus->Unsubscribe(AsShared(), TopicPattern);
+	}
+}
+
+MESSAGING_API void FMessageEndpoint::ProcessMessage(const IMessageContextRef& Context)
+{
+	if (!Context->IsValid())
+	{
+		return;
+	}
+
+	for (int32 HandlerIndex = 0; HandlerIndex < Handlers.Num(); ++HandlerIndex)
+	{
+		IMessageHandlerPtr& Handler = Handlers[HandlerIndex];
+
+		if (Handler->GetHandledMessageType() == Context->GetMessageType())
+		{
+			Handler->HandleMessage(Context);
+		}
+	}
+}
\ No newline at end of file
