Add copyright notes to new files.

diff --git a/Engine/Source/Runtime/Messaging/Private/MessageEndpoint.cpp b/Engine/Source/Runtime/Messaging/Private/MessageEndpoint.cpp
--- a/Engine/Source/Runtime/Messaging/Private/MessageEndpoint.cpp
+++ b/Engine/Source/Runtime/Messaging/Private/MessageEndpoint.cpp
@@ -1,3 +1,4 @@
+// Copyright 1998-2014 Epic Games, Inc. All Rights Reserved.
 #include "MessagingPrivatePCH.h"
 #include "MessageEndpoint.h"
 
diff --git a/Engine/Source/Runtime/Messaging/Private/MessageEndpointBuilder.cpp b/Engine/Source/Runtime/Messaging/Private/MessageEndpointBuilder.cpp
--- a/Engine/Source/Runtime/Messaging/Private/MessageEndpointBuilder.cpp
+++ b/Engine/Source/Runtime/Messaging/Private/MessageEndpointBuilder.cpp
@@ -1,3 +1,4 @@
+// Copyright 1998-2014 Epic Games, Inc. All Rights Reserved.
 #include "MessagingPrivatePCH.h"
 #include "MessageEndpointBuilder.h"
 
