Back out splitting MessageEndpoint and MessageEndpointBuilder to .cpp and header.

diff --git a/Engine/Source/Developer/AutomationController/AutomationController.Build.cs b/Engine/Source/Developer/AutomationController/AutomationController.Build.cs
--- a/Engine/Source/Developer/AutomationController/AutomationController.Build.cs
+++ b/Engine/Source/Developer/AutomationController/AutomationController.Build.cs
@@ -18,7 +18,6 @@ namespace UnrealBuildTool.Rules
 				{
 					"AutomationMessages",
 					"CoreUObject",
-					"Messaging",
 					"UnrealEdMessages",
 				}
 			);
diff --git a/Engine/Source/Developer/LauncherServices/LauncherServices.Build.cs b/Engine/Source/Developer/LauncherServices/LauncherServices.Build.cs
--- a/Engine/Source/Developer/LauncherServices/LauncherServices.Build.cs
+++ b/Engine/Source/Developer/LauncherServices/LauncherServices.Build.cs
@@ -21,7 +21,6 @@ namespace UnrealBuildTool.Rules
 			PrivateDependencyModuleNames.AddRange(
 				new string[] {
 					"CoreUObject",
-					"Messaging",
 					"SessionMessages",
 					"TargetPlatform",
 					"UnrealEdMessages",
diff --git a/Engine/Source/Developer/ProfilerClient/ProfilerClient.Build.cs b/Engine/Source/Developer/ProfilerClient/ProfilerClient.Build.cs
--- a/Engine/Source/Developer/ProfilerClient/ProfilerClient.Build.cs
+++ b/Engine/Source/Developer/ProfilerClient/ProfilerClient.Build.cs
@@ -19,7 +19,6 @@ namespace UnrealBuildTool.Rules
 				new string[]
 				{
 					"CoreUObject",
-					"Messaging",
 					"ProfilerMessages",
 				}
 			);
diff --git a/Engine/Source/Developer/ProfilerService/ProfilerService.Build.cs b/Engine/Source/Developer/ProfilerService/ProfilerService.Build.cs
--- a/Engine/Source/Developer/ProfilerService/ProfilerService.Build.cs
+++ b/Engine/Source/Developer/ProfilerService/ProfilerService.Build.cs
@@ -18,7 +18,6 @@ namespace UnrealBuildTool.Rules
 				{
 					"CoreUObject",
 					"Engine",
-					"Messaging",
 					"ProfilerMessages",
 				}
 			);
diff --git a/Engine/Source/Developer/ScreenShotComparisonTools/ScreenShotComparisonTools.Build.cs b/Engine/Source/Developer/ScreenShotComparisonTools/ScreenShotComparisonTools.Build.cs
--- a/Engine/Source/Developer/ScreenShotComparisonTools/ScreenShotComparisonTools.Build.cs
+++ b/Engine/Source/Developer/ScreenShotComparisonTools/ScreenShotComparisonTools.Build.cs
@@ -16,7 +16,6 @@ public class ScreenShotComparisonTools : ModuleRules
 			new string[] {
 				"AutomationMessages",
 				"CoreUObject",
-				"Messaging",
 				"UnrealEdMessages",
 				"Slate",
                 "EditorStyle",
diff --git a/Engine/Source/Developer/SessionServices/SessionServices.Build.cs b/Engine/Source/Developer/SessionServices/SessionServices.Build.cs
--- a/Engine/Source/Developer/SessionServices/SessionServices.Build.cs
+++ b/Engine/Source/Developer/SessionServices/SessionServices.Build.cs
@@ -18,7 +18,6 @@ public class SessionServices : ModuleRules
 			{
 				"CoreUObject",
 				"EngineMessages",
-				"Messaging",
 				"SessionMessages",
 			}
 		);
diff --git a/Engine/Source/Developer/TargetDeviceServices/TargetDeviceServices.Build.cs b/Engine/Source/Developer/TargetDeviceServices/TargetDeviceServices.Build.cs
--- a/Engine/Source/Developer/TargetDeviceServices/TargetDeviceServices.Build.cs
+++ b/Engine/Source/Developer/TargetDeviceServices/TargetDeviceServices.Build.cs
@@ -18,7 +18,6 @@ namespace UnrealBuildTool.Rules
 				new string[]
 				{
 					"CoreUObject",
-					"Messaging",
 				}
 			);
 
diff --git a/Engine/Source/Developer/iOS/IOSTargetPlatform/IOSTargetPlatform.Build.cs b/Engine/Source/Developer/iOS/IOSTargetPlatform/IOSTargetPlatform.Build.cs
--- a/Engine/Source/Developer/iOS/IOSTargetPlatform/IOSTargetPlatform.Build.cs
+++ b/Engine/Source/Developer/iOS/IOSTargetPlatform/IOSTargetPlatform.Build.cs
@@ -13,7 +13,6 @@ public class IOSTargetPlatform : ModuleRules
 			{
 				"Core",
 				"CoreUObject",
-				"Messaging",
 				"TargetPlatform",
 				"DesktopPlatform",
 				"LaunchDaemonMessages",
diff --git a/Engine/Source/Editor/UnrealEd/UnrealEd.Build.cs b/Engine/Source/Editor/UnrealEd/UnrealEd.Build.cs
--- a/Engine/Source/Editor/UnrealEd/UnrealEd.Build.cs
+++ b/Engine/Source/Editor/UnrealEd/UnrealEd.Build.cs
@@ -87,7 +87,6 @@ public class UnrealEd : ModuleRules
 				"LauncherServices",
 				"MaterialEditor",
 				"MessageLog",
-				"Messaging",
 				"NetworkFileSystem",
 				"ProjectSettingsViewer",
 				"PropertyEditor",
diff --git a/Engine/Source/Runtime/AutomationWorker/AutomationWorker.Build.cs b/Engine/Source/Runtime/AutomationWorker/AutomationWorker.Build.cs
--- a/Engine/Source/Runtime/AutomationWorker/AutomationWorker.Build.cs
+++ b/Engine/Source/Runtime/AutomationWorker/AutomationWorker.Build.cs
@@ -18,7 +18,6 @@ namespace UnrealBuildTool.Rules
 				{
 					"AutomationMessages",
 					"CoreUObject",
-					"Messaging",
 				}
 			);
 
diff --git a/Engine/Source/Runtime/Messaging/Private/MessageEndpoint.cpp b/Engine/Source/Runtime/Messaging/Private/MessageEndpoint.cpp
deleted file mode 100644
diff --git a/Engine/Source/Runtime/Messaging/Private/MessageEndpointBuilder.cpp b/Engine/Source/Runtime/Messaging/Private/MessageEndpointBuilder.cpp
deleted file mode 100644
diff --git a/Engine/Source/Runtime/Messaging/Public/Common/MessageEndpoint.h b/Engine/Source/Runtime/Messaging/Public/Common/MessageEndpoint.h
--- a/Engine/Source/Runtime/Messaging/Public/Common/MessageEndpoint.h
+++ b/Engine/Source/Runtime/Messaging/Public/Common/MessageEndpoint.h
@@ -68,12 +68,28 @@ public:
 	 * @param InBus The message bus to attach this endpoint to.
 	 * @param InHandlers The collection of message handlers to register.
 	 */
-	MESSAGING_API FMessageEndpoint( const FName& InName, const IMessageBusRef& InBus, const TArray<IMessageHandlerPtr>& InHandlers );
+	FMessageEndpoint( const FName& InName, const IMessageBusRef& InBus, const TArray<IMessageHandlerPtr>& InHandlers )
+		: Address(FGuid::NewGuid())
+		, BusPtr(InBus)
+		, Enabled(true)
+		, Handlers(InHandlers)
+		, Id(FGuid::NewGuid())
+		, InboxEnabled(false)
+		, Name(InName)
+	{ }
 
 	/**
 	 * Destructor.
 	 */
-	MESSAGING_API ~FMessageEndpoint();
+	~FMessageEndpoint( )
+	{
+		IMessageBusPtr Bus = BusPtr.Pin();
+
+		if (Bus.IsValid())
+		{
+			Bus->Unregister(Address);
+		}
+	}
 
 public:
 
@@ -161,7 +177,15 @@ public:
 	 * @param Message The context of the message to defer.
 	 * @param Delay The time delay.
 	 */
-	MESSAGING_API void Defer(const IMessageContextRef& Context, const FTimespan& Delay);
+	void Defer( const IMessageContextRef& Context, const FTimespan& Delay )
+	{
+		IMessageBusPtr Bus = GetBusIfEnabled();
+
+		if (Bus.IsValid())
+		{
+			Bus->Forward(Context, TArrayBuilder<FMessageAddress>().Add(Address), Context->GetScope(), Delay, AsShared());
+		}
+	}
 
 	/**
 	 * Forwards a previously received message.
@@ -173,7 +197,15 @@ public:
 	 * @param ForwardingScope The scope of the forwarded message.
 	 * @param Delay The time delay.
 	 */
-	MESSAGING_API void Forward(const IMessageContextRef& Context, const TArray<FMessageAddress>& Recipients, EMessageScope::Type ForwardingScope, const FTimespan& Delay);
+	void Forward( const IMessageContextRef& Context, const TArray<FMessageAddress>& Recipients, EMessageScope::Type ForwardingScope, const FTimespan& Delay )
+	{
+		IMessageBusPtr Bus = GetBusIfEnabled();
+
+		if (Bus.IsValid())
+		{
+			Bus->Forward(Context, Recipients, ForwardingScope, Delay, AsShared());
+		}
+	}
 
 	/**
 	 * Publishes a message to all subscribed recipients within the specified scope.
@@ -185,7 +217,15 @@ public:
 	 * @param Delay The delay after which to publish the message.
 	 * @param Expiration The time at which the message expires.
 	 */
-	MESSAGING_API void Publish(void* Message, UScriptStruct* TypeInfo, EMessageScope::Type Scope, const FTimespan& Delay, const FDateTime& Expiration);
+	void Publish( void* Message, UScriptStruct* TypeInfo, EMessageScope::Type Scope, const FTimespan& Delay, const FDateTime& Expiration )
+	{
+		IMessageBusPtr Bus = GetBusIfEnabled();
+
+		if (Bus.IsValid())
+		{
+			Bus->Publish(Message, TypeInfo, Scope, Delay, Expiration, AsShared());
+		}
+	}
 
 	/**
 	 * Sends a message to the specified list of recipients.
@@ -197,7 +237,15 @@ public:
 	 * @param Delay The delay after which to send the message.
 	 * @param Expiration The time at which the message expires.
 	 */
-	MESSAGING_API void Send(void* Message, UScriptStruct* TypeInfo, const IMessageAttachmentPtr& Attachment, const TArray<FMessageAddress>& Recipients, const FTimespan& Delay, const FDateTime& Expiration);
+	void Send( void* Message, UScriptStruct* TypeInfo, const IMessageAttachmentPtr& Attachment, const TArray<FMessageAddress>& Recipients, const FTimespan& Delay, const FDateTime& Expiration )
+	{
+		IMessageBusPtr Bus = GetBusIfEnabled();
+
+		if (Bus.IsValid())
+		{
+			Bus->Send(Message, TypeInfo, Attachment, Recipients, Delay, Expiration, AsShared());
+		}
+	}
 
 	/**
 	 * Subscribes a message handler.
@@ -205,7 +253,15 @@ public:
 	 * @param MessageType The type name of the messages to subscribe to.
 	 * @param ScopeRange The range of message scopes to include in the subscription.
 	 */
-	MESSAGING_API void Subscribe(const FName& MessageType, const FMessageScopeRange& ScopeRange);
+	void Subscribe( const FName& MessageType, const FMessageScopeRange& ScopeRange )
+	{
+		IMessageBusPtr Bus = GetBusIfEnabled();
+
+		if (Bus.IsValid())
+		{
+			Bus->Subscribe(AsShared(), MessageType, ScopeRange);
+		}
+	}
 
 	/**
 	 * Unsubscribes this endpoint from the specified message type.
@@ -213,7 +269,15 @@ public:
 	 * @param MessageType The type of message to unsubscribe (NAME_All = all types).
 	 * @see Subscribe
 	 */
-	MESSAGING_API void Unsubscribe(const FName& TopicPattern);
+	void Unsubscribe( const FName& TopicPattern )
+	{
+		IMessageBusPtr Bus = GetBusIfEnabled();
+
+		if (Bus.IsValid())
+		{
+			Bus->Unsubscribe(AsShared(), TopicPattern);
+		}
+	}
 
 public:
 
@@ -711,7 +775,23 @@ protected:
 	 *
 	 * @param Context The context of the message to handle.
 	 */
-	MESSAGING_API void ProcessMessage(const IMessageContextRef& Context);
+	void ProcessMessage( const IMessageContextRef& Context )
+	{
+		if (!Context->IsValid())
+		{
+			return;
+		}
+
+		for (int32 HandlerIndex = 0; HandlerIndex < Handlers.Num(); ++HandlerIndex)
+		{
+			IMessageHandlerPtr& Handler = Handlers[HandlerIndex];
+
+			if (Handler->GetHandledMessageType() == Context->GetMessageType())
+			{
+				Handler->HandleMessage(Context);
+			}
+		}
+	}
 
 private:
 
diff --git a/Engine/Source/Runtime/Messaging/Public/Common/MessageEndpointBuilder.h b/Engine/Source/Runtime/Messaging/Public/Common/MessageEndpointBuilder.h
--- a/Engine/Source/Runtime/Messaging/Public/Common/MessageEndpointBuilder.h
+++ b/Engine/Source/Runtime/Messaging/Public/Common/MessageEndpointBuilder.h
@@ -150,7 +150,35 @@ public:
 	 *
 	 * @return A new message endpoint, or nullptr if it couldn't be built.
 	 */
-	MESSAGING_API FMessageEndpointPtr Build();
+	FMessageEndpointPtr Build( )
+	{
+		FMessageEndpointPtr Endpoint;
+
+		IMessageBusPtr Bus = BusPtr.Pin();
+		
+		if (Bus.IsValid())
+		{
+			Endpoint = MakeShareable(new FMessageEndpoint(Name, Bus.ToSharedRef(), Handlers));
+			Bus->Register(Endpoint->GetAddress(), Endpoint.ToSharedRef());
+
+			if (Disabled)
+			{
+				Endpoint->Disable();
+			}
+
+			if (InboxEnabled)
+			{
+				Endpoint->EnableInbox();
+				Endpoint->SetRecipientThread(ENamedThreads::AnyThread);
+			}
+			else
+			{
+				Endpoint->SetRecipientThread(RecipientThread);
+			}
+		}
+
+		return Endpoint;
+	}
 
 	/**
 	 * Implicit conversion operator to build the message endpoint as configured.
