Messaging: Documentation cleanup pass

diff --git a/Engine/Source/Runtime/Messaging/Private/Bridge/MessageAddressBook.h b/Engine/Source/Runtime/Messaging/Private/Bridge/MessageAddressBook.h
--- a/Engine/Source/Runtime/Messaging/Private/Bridge/MessageAddressBook.h
+++ b/Engine/Source/Runtime/Messaging/Private/Bridge/MessageAddressBook.h
@@ -10,17 +10,13 @@ class FMessageAddressBook
 {
 public:
 
-	/**
-	 * Default constructor.
-	 */
+	/** Default constructor. */
 	FMessageAddressBook( )
 	{
 		CriticalSection = new FCriticalSection();
 	}
 
-	/**
-	 * Destructor.
-	 */
+	/** Destructor. */
 	~FMessageAddressBook( )
 	{
 		delete CriticalSection;
@@ -41,9 +37,7 @@ public:
 		Entries.FindOrAdd(Address) = NodeId;
 	}
 
-	/**
-	 * Clears the address book.
-	 */
+	/** Clears the address book. */
 	void Clear( )
 	{
 		FScopeLock Lock(CriticalSection);
@@ -137,9 +131,9 @@ public:
 
 private:
 
-	// Holds a critical section to serialize access to the address book entries.
+	/** Holds a critical section to serialize access to the address book entries. */
 	FCriticalSection* CriticalSection;
 
-	// Holds the collection of known addresses and their transport node identifiers.
+	/** Holds the collection of known addresses and their transport node identifiers. */
 	TMap<FMessageAddress, FGuid> Entries;
 };
diff --git a/Engine/Source/Runtime/Messaging/Private/Bridge/MessageBridge.h b/Engine/Source/Runtime/Messaging/Private/Bridge/MessageBridge.h
--- a/Engine/Source/Runtime/Messaging/Private/Bridge/MessageBridge.h
+++ b/Engine/Source/Runtime/Messaging/Private/Bridge/MessageBridge.h
@@ -24,9 +24,7 @@ public:
 	 */
 	FMessageBridge( const FMessageAddress InAddress, const IMessageBusRef& InBus, const ISerializeMessagesRef& InSerializer, const ITransportMessagesRef& InTransport );
 
-	/**
-	 * Destructor.
-	 */
+	/** Destructor. */
 	~FMessageBridge( );
 
 public:
@@ -98,27 +96,27 @@ private:
 
 private:
 
-	// Holds the bridge's address.
+	/** Holds the bridge's address. */
 	FMessageAddress Address;
 
-	// Holds the address book.
+	/** Holds the address book. */
 	FMessageAddressBook AddressBook;
 
-	// Holds a reference to the bus that this bridge is attached to.
+	/** Holds a reference to the bus that this bridge is attached to. */
 	IMessageBusPtr Bus;
 
-	// Hold a flag indicating whether this endpoint is active.
+	/** Hold a flag indicating whether this endpoint is active. */
 	bool Enabled;
 
-	// Holds the bridge's unique identifier (for debugging purposes).
+	/** Holds the bridge's unique identifier (for debugging purposes). */
 	const FGuid Id;
 
-	// Holds the message subscription for outbound messages.
+	/** Holds the message subscription for outbound messages. */
 	IMessageSubscriptionPtr MessageSubscription;
 
-	// Holds the message serializer. 
+	/** Holds the message serializer. */
 	ISerializeMessagesPtr Serializer;
 
-	// Holds the message transport object.
+	/** Holds the message transport object. */
 	ITransportMessagesPtr Transport;
 };
diff --git a/Engine/Source/Runtime/Messaging/Private/Bridge/MessageData.h b/Engine/Source/Runtime/Messaging/Private/Bridge/MessageData.h
--- a/Engine/Source/Runtime/Messaging/Private/Bridge/MessageData.h
+++ b/Engine/Source/Runtime/Messaging/Private/Bridge/MessageData.h
@@ -19,9 +19,7 @@ class FMessageData
 {
 public:
 
-	/**
-	 * Default constructor.
-	 */
+	/** Default constructor. */
 	FMessageData( )
 		: FMemoryWriter(Data, true)
 		, State(EMessageDataState::Incomplete)
@@ -69,6 +67,6 @@ private:
 
 private:
 
-	// Holds a delegate that is invoked when the data's state changed.
+	/** Holds a delegate that is invoked when the data's state changed. */
 	FSimpleDelegate StateChangedDelegate;
 };
diff --git a/Engine/Source/Runtime/Messaging/Private/Bridge/MessageDeserializeTask.h b/Engine/Source/Runtime/Messaging/Private/Bridge/MessageDeserializeTask.h
--- a/Engine/Source/Runtime/Messaging/Private/Bridge/MessageDeserializeTask.h
+++ b/Engine/Source/Runtime/Messaging/Private/Bridge/MessageDeserializeTask.h
@@ -82,12 +82,12 @@ public:
 
 private:
 
-	// Holds the message's type information.
+	/** Holds the message's type information. */
 	TWeakObjectPtr<UScriptStruct> MessageTypeInfoPtr;
 
-	// Holds a reference to the serializer to use.
+	/** Holds a reference to the serializer to use. */
 	ISerializeMessagesRef Serializer;
 
-	// Holds the name of the thread that deserialization should be run on.
+	/** Holds the name of the thread that deserialization should be run on. */
 	ENamedThreads::Type Thread;
 };
diff --git a/Engine/Source/Runtime/Messaging/Private/Bridge/MessageForwardTask.h b/Engine/Source/Runtime/Messaging/Private/Bridge/MessageForwardTask.h
--- a/Engine/Source/Runtime/Messaging/Private/Bridge/MessageForwardTask.h
+++ b/Engine/Source/Runtime/Messaging/Private/Bridge/MessageForwardTask.h
@@ -72,12 +72,12 @@ public:
 
 private:
 
-	// Holds a pointer to the message bus to forward the message to.
+	/** Holds a pointer to the message bus to forward the message to. */
 	IMessageBusWeakPtr BusPtr;
 
-	// Holds the context of the message to forward.
+	/** Holds the context of the message to forward. */
 	IMessageContextRef MessageContext;
 
-	// Holds a reference to the message sender.
+	/** Holds a reference to the message sender. */
 	ISendMessagesRef Sender;
 };
diff --git a/Engine/Source/Runtime/Messaging/Private/Bridge/MessageSerializeTask.h b/Engine/Source/Runtime/Messaging/Private/Bridge/MessageSerializeTask.h
--- a/Engine/Source/Runtime/Messaging/Private/Bridge/MessageSerializeTask.h
+++ b/Engine/Source/Runtime/Messaging/Private/Bridge/MessageSerializeTask.h
@@ -79,12 +79,12 @@ public:
 
 private:
 
-	// Holds the context of the message to serialize.
+	/** Holds the context of the message to serialize. */
 	IMessageContextRef MessageContext;
 
-	// Holds a reference to the message data.
+	/** Holds a reference to the message data. */
 	FMessageDataRef MessageData;
 
-	// Holds a reference to the serializer to use.
+	/** Holds a reference to the serializer to use. */
 	TWeakPtr<ISerializeMessages, ESPMode::ThreadSafe> SerializerPtr;
 };
diff --git a/Engine/Source/Runtime/Messaging/Private/Bus/MessageBus.h b/Engine/Source/Runtime/Messaging/Private/Bus/MessageBus.h
--- a/Engine/Source/Runtime/Messaging/Private/Bus/MessageBus.h
+++ b/Engine/Source/Runtime/Messaging/Private/Bus/MessageBus.h
@@ -19,9 +19,7 @@ public:
 	 */
 	FMessageBus( const IAuthorizeMessageRecipientsPtr& InRecipientAuthorizer );
 
-	/**
-	 * Destructor
-	 */
+	/** Destructor */
 	~FMessageBus( );
 
 public:
@@ -116,15 +114,15 @@ public:
 
 private:
 
-	// Holds the message router.
+	/** Holds the message router. */
 	FMessageRouter* Router;
 
-	// Holds the message router thread.
+	/** Holds the message router thread. */
 	FRunnableThread* RouterThread;
 
-	// Holds the recipient authorizer.
+	/** Holds the recipient authorizer. */
 	IAuthorizeMessageRecipientsPtr RecipientAuthorizer;
 
-	// Holds bus shutdown delegate.
+	/** Holds bus shutdown delegate. */
 	FOnMessageBusShutdown ShutdownDelegate;
 };
diff --git a/Engine/Source/Runtime/Messaging/Private/Bus/MessageContext.h b/Engine/Source/Runtime/Messaging/Private/Bus/MessageContext.h
--- a/Engine/Source/Runtime/Messaging/Private/Bus/MessageContext.h
+++ b/Engine/Source/Runtime/Messaging/Private/Bus/MessageContext.h
@@ -14,9 +14,7 @@ class FMessageContext
 {
 public:
 
-	/**
-	 * Default constructor.
-	 */
+	/** Default constructor. */
 	FMessageContext( )
 		: Message(nullptr)
 		, TypeInfo(nullptr)
@@ -70,9 +68,7 @@ public:
 		, TypeInfo(nullptr)
 	{ }
 
-	/**
-	 * Destructor.
-	 */
+	/** Destructor. */
 	~FMessageContext()
 	{
 		if ((Message != nullptr) && TypeInfo.IsValid())
@@ -257,42 +253,42 @@ public:
 
 private:
 
-	// Holds a pointer to attached binary data.
+	/** Holds a pointer to attached binary data. */
 	IMessageAttachmentPtr Attachment;
 
-	// Holds the expiration time.
+	/** Holds the expiration time. */
 	FDateTime Expiration;
 
-	// Holds the address of the endpoint that forwarded this message.
+	/** Holds the address of the endpoint that forwarded this message. */
 	FMessageAddress Forwarder;
 
-	// Holds the optional message headers.
+	/** Holds the optional message headers. */
 	TMap<FName, FString> Headers;
 
-	// Holds the message.
+	/** Holds the message. */
 	void* Message;
 
-	// Holds the original message context.
+	/** Holds the original message context. */
 	IMessageContextPtr OriginalContext;
 
-	// Holds the message recipients.
+	/** Holds the message recipients. */
 	TArray<FMessageAddress> Recipients;
 
-	// Holds the message's scope.
+	/** Holds the message's scope. */
 	EMessageScope::Type Scope;
 
-	// Holds the sender's identifier.
+	/** Holds the sender's identifier. */
 	FMessageAddress Sender;
 
-	// Holds the name of the thread from which the message was sent.
+	/** Holds the name of the thread from which the message was sent. */
 	ENamedThreads::Type SenderThread;
 
-	// Holds the time at which the message was forwarded.
+	/** Holds the time at which the message was forwarded. */
 	FDateTime TimeForwarded;
 
-	// Holds the time at which the message was sent.
+	/** Holds the time at which the message was sent. */
 	FDateTime TimeSent;
 
-	// Holds the message's type information.
+	/** Holds the message's type information. */
 	TWeakObjectPtr<UScriptStruct> TypeInfo;
 };
diff --git a/Engine/Source/Runtime/Messaging/Private/Bus/MessageDispatchTask.h b/Engine/Source/Runtime/Messaging/Private/Bus/MessageDispatchTask.h
--- a/Engine/Source/Runtime/Messaging/Private/Bus/MessageDispatchTask.h
+++ b/Engine/Source/Runtime/Messaging/Private/Bus/MessageDispatchTask.h
@@ -84,15 +84,15 @@ public:
 
 private:
 
-	// Holds the message context.
+	/** Holds the message context. */
 	IMessageContextRef Context;
 
-	// Holds a reference to the recipient.
+	/** Holds a reference to the recipient. */
 	IReceiveMessagesWeakPtr RecipientPtr;
 
-	// Holds the name of the thread that the router is running on.
+	/** Holds the name of the thread that the router is running on. */
 	ENamedThreads::Type Thread;
 
-	// Holds a pointer to the message tracer.
+	/** Holds a pointer to the message tracer. */
 	FMessageTracerWeakPtr TracerPtr;
 };
diff --git a/Engine/Source/Runtime/Messaging/Private/Bus/MessageRouter.h b/Engine/Source/Runtime/Messaging/Private/Bus/MessageRouter.h
--- a/Engine/Source/Runtime/Messaging/Private/Bus/MessageRouter.h
+++ b/Engine/Source/Runtime/Messaging/Private/Bus/MessageRouter.h
@@ -11,14 +11,10 @@ class FMessageRouter
 {
 public:
 
-	/**
-	 * Default constructor.
-	 */
+	/** Default constructor. */
 	FMessageRouter( );
 
-	/**
-	 * Destructor.
-	 */
+	/** Destructor. */
 	~FMessageRouter( );
 
 public:
@@ -161,9 +157,7 @@ protected:
 	 */
 	void DispatchMessage( const IMessageContextRef& Message );
 
-	/**
-	 * Processes all delayed messages.
-	 */
+	/** Processes all delayed messages. */
 	void ProcessDelayedMessages( );
 
 private:
@@ -203,56 +197,56 @@ private:
 
 private:
 
-	// Handles adding message interceptors.
+	/** Handles adding message interceptors. */
 	void HandleAddInterceptor( IInterceptMessagesRef Interceptor, FName MessageType );
 
-	// Handles adding message recipients.
+	/** Handles adding message recipients. */
 	void HandleAddRecipient( FMessageAddress Address, IReceiveMessagesWeakPtr RecipientPtr );
 
-	// Handles adding of subscriptions.
+	/** Handles adding of subscriptions. */
 	void HandleAddSubscriber( IMessageSubscriptionRef Subscription );
 
-	// Handles the removal of message interceptors.
+	/** Handles the removal of message interceptors. */
 	void HandleRemoveInterceptor( IInterceptMessagesRef Interceptor, FName MessageType );
 
-	// Handles the removal of message recipients.
+	/** Handles the removal of message recipients. */
 	void HandleRemoveRecipient( FMessageAddress Address );
 
-	// Handles the removal of subscribers.
+	/** Handles the removal of subscribers. */
 	void HandleRemoveSubscriber( IReceiveMessagesWeakPtr SubscriberPtr, FName MessageType );
 
-	// Handles the routing of messages.
+	/** Handles the routing of messages. */
 	void HandleRouteMessage( IMessageContextRef Context );
 
 private:
 
-	// Maps message types to interceptors.
+	/** Maps message types to interceptors. */
 	TMap<FName, TArray<IInterceptMessagesPtr>> ActiveInterceptors;
 
-	// Maps message addresses to recipients.
+	/** Maps message addresses to recipients. */
 	TMap<FMessageAddress, IReceiveMessagesWeakPtr> ActiveRecipients;
 
-	// Maps message types to subscriptions.
+	/** Maps message types to subscriptions. */
 	TMap<FName, TArray<IMessageSubscriptionPtr>> ActiveSubscriptions;
 
-	// Holds the router command queue.
+	/** Holds the router command queue. */
 	TQueue<TBaseDelegate_NoParams<void>, EQueueMode::Mpsc> Commands;
 
-	// Holds the current time.
+	/** Holds the current time. */
 	FDateTime CurrentTime;
 
-	// Holds the collection of delayed messages.
+	/** Holds the collection of delayed messages. */
 	TArray<FDelayedMessage> DelayedMessages;
 
-	// Holds a sequence number for delayed messages.
+	/** Holds a sequence number for delayed messages. */
 	int64 DelayedMessagesSequence;
 
-	// Holds a flag indicating that the thread is stopping.
+	/** Holds a flag indicating that the thread is stopping. */
 	bool Stopping;
 
-	// Holds the message tracer.
+	/** Holds the message tracer. */
 	FMessageTracerRef Tracer;
 
-	// Holds an event signaling that work is available.
+	/** Holds an event signaling that work is available. */
 	FEvent* WorkEvent;
 };
diff --git a/Engine/Source/Runtime/Messaging/Private/Bus/MessageSubscription.h b/Engine/Source/Runtime/Messaging/Private/Bus/MessageSubscription.h
--- a/Engine/Source/Runtime/Messaging/Private/Bus/MessageSubscription.h
+++ b/Engine/Source/Runtime/Messaging/Private/Bus/MessageSubscription.h
@@ -66,15 +66,15 @@ public:
 
 private:
 
-	// Holds a flag indicating whether this subscription is enabled.
+	/** Holds a flag indicating whether this subscription is enabled. */
 	bool Enabled;
 
-	// Holds the type of subscribed messages.
+	/** Holds the type of subscribed messages. */
 	FName MessageType;
 
-	// Holds the range of message scopes to subscribe to.
+	/** Holds the range of message scopes to subscribe to. */
 	FMessageScopeRange ScopeRange;
 
-	// Holds the subscriber.
+	/** Holds the subscriber. */
 	IReceiveMessagesWeakPtr Subscriber;
 };
diff --git a/Engine/Source/Runtime/Messaging/Private/Bus/MessageTracer.h b/Engine/Source/Runtime/Messaging/Private/Bus/MessageTracer.h
--- a/Engine/Source/Runtime/Messaging/Private/Bus/MessageTracer.h
+++ b/Engine/Source/Runtime/Messaging/Private/Bus/MessageTracer.h
@@ -19,19 +19,19 @@ typedef TSharedRef<class FMessageTracer, ESPMode::ThreadSafe> FMessageTracerRef;
 class FMessageTracer
 	: public IMessageTracer
 {
-	// Holds debug information for message recipients.
+	/** Holds debug information for message recipients. */
 	struct FRecipientInfo
 	{
-		// Holds the recipient's unique identifier.
+		/** Holds the recipient's unique identifier. */
 		FGuid Id;
 
-		// Holds the recipient's name.
+		/** Holds the recipient's name. */
 		FName Name;
 
-		// Holds a flag indicating whether the recipient is a remote endpoint.
+		/** Holds a flag indicating whether the recipient is a remote endpoint. */
 		bool Remote;
 
-		// Creates and initializes a new instance.
+		/** Creates and initializes a new instance. */
 		FRecipientInfo( const FMessageAddress& Address, const IReceiveMessagesRef& Recipient )
 			: Id(Recipient->GetRecipientId())
 			, Name(Recipient->GetDebugName())
@@ -41,14 +41,10 @@ class FMessageTracer
 
 public:
 
-	/**
-	 * Default constructor.
-	 */
+	/** Default constructor. */
 	FMessageTracer( );
 
-	/**
-	 * Destructor.
-	 */
+	/** Destructor. */
 	~FMessageTracer( );
 
 public:
@@ -425,47 +421,47 @@ protected:
 
 private:
 
-	// Holds the collection of endpoints for known message addresses.
+	/** Holds the collection of endpoints for known message addresses. */
 	TMap<FMessageAddress, FMessageTracerEndpointInfoPtr> AddressesToEndpointInfos;
 
-	// Holds a flag indicating whether a breakpoint was hit.
+	/** Holds a flag indicating whether a breakpoint was hit. */
 	bool Breaking;
 
-	// Holds the collection of breakpoints.
+	/** Holds the collection of breakpoints. */
 	TArray<IMessageTracerBreakpointPtr> Breakpoints;
 
-	// Holds the collection of message senders to break on when they send a message.
+	/** Holds the collection of message senders to break on when they send a message. */
 	TArray<FMessageAddress> BreakOnSenders;
 
-	// Holds an event signaling that messaging routing can continue.
+	/** Holds an event signaling that messaging routing can continue. */
 	FEvent* ContinueEvent;
 
-	// Holds the collection of endpoints for known recipient identifiers.
+	/** Holds the collection of endpoints for known recipient identifiers. */
 	TMap<FGuid, FMessageTracerEndpointInfoPtr> RecipientsToEndpointInfos;
 
-	// Holds the collection of known messages.
+	/** Holds the collection of known messages. */
 	TMap<IMessageContextPtr, FMessageTracerMessageInfoPtr> MessageInfos;
 
-	// Holds the collection of known message types.
+	/** Holds the collection of known message types. */
 	TMap<FName, FMessageTracerTypeInfoPtr> MessageTypes;
 
-	// Holds a flag indicating whether a reset is pending.
+	/** Holds a flag indicating whether a reset is pending. */
 	bool ResetPending;
 
-	// Holds a flag indicating whether the tracer is running.
+	/** Holds a flag indicating whether the tracer is running. */
 	bool Running;
 
-	// Holds the trace actions queue.
+	/** Holds the trace actions queue. */
 	TQueue<TBaseDelegate_NoParams<void>, EQueueMode::Mpsc> Traces;
 
 private:
 
-	// Holds a delegate that is executed when a new message has been added to the collection of known messages.
+	/** Holds a delegate that is executed when a new message has been added to the collection of known messages. */
 	FMessageTracerMessageAdded MessagesAddedDelegate;
 
-	// Holds a delegate that is executed when the message history has been reset.
+	/** Holds a delegate that is executed when the message history has been reset. */
 	FSimpleMulticastDelegate MessagesResetDelegate;
 
-	// Holds a delegate that is executed when a new type has been added to the collection of known message types.
+	/** Holds a delegate that is executed when a new type has been added to the collection of known message types. */
 	FMessageTracerTypeAdded TypeAddedDelegate;
 };
diff --git a/Engine/Source/Runtime/Messaging/Public/Common/FileMessageAttachment.h b/Engine/Source/Runtime/Messaging/Public/Common/FileMessageAttachment.h
--- a/Engine/Source/Runtime/Messaging/Public/Common/FileMessageAttachment.h
+++ b/Engine/Source/Runtime/Messaging/Public/Common/FileMessageAttachment.h
@@ -32,9 +32,7 @@ public:
 		, Filename(InFilename)
 	{ }
 
-	/**
-	 * Destructor.
-	 */
+	/** Destructor. */
 	~FFileMessageAttachment( )
 	{
 		if (AutoDeleteFile)
@@ -54,9 +52,9 @@ public:
 
 private:
 
-	// Holds a flag indicating whether the file should be deleted.
+	/** Holds a flag indicating whether the file should be deleted. */
 	bool AutoDeleteFile;
 
-	// Holds the name of the file that holds the attached data.
+	/** Holds the name of the file that holds the attached data. */
 	FString Filename;
 };
diff --git a/Engine/Source/Runtime/Messaging/Public/Common/MessageBridgeBuilder.h b/Engine/Source/Runtime/Messaging/Public/Common/MessageBridgeBuilder.h
--- a/Engine/Source/Runtime/Messaging/Public/Common/MessageBridgeBuilder.h
+++ b/Engine/Source/Runtime/Messaging/Public/Common/MessageBridgeBuilder.h
@@ -10,9 +10,7 @@ class FMessageBridgeBuilder
 {
 public:
 
-	/**
-	 * Default constructor.
-	 */
+	/** Default constructor. */
 	FMessageBridgeBuilder( )
 		: Address(FMessageAddress::NewGuid())
 		, BusPtr(IMessagingModule::Get().GetDefaultBus())
@@ -146,18 +144,18 @@ public:
 
 private:
 
-	// Holds the bridge's address.
+	/** Holds the bridge's address. */
 	FMessageAddress Address;
 
-	// Holds a weak pointer to the message bus to attach to.
+	/** Holds a weak pointer to the message bus to attach to. */
 	IMessageBusWeakPtr BusPtr;
 
-	// Holds a flag indicating whether the bridge should be disabled.
+	/** Holds a flag indicating whether the bridge should be disabled. */
 	bool Disabled;
 
-	// Holds a reference to the message serializer.
+	/** Holds a reference to the message serializer. */
 	ISerializeMessagesPtr Serializer;
 
-	// Holds a reference to the message transport technology.
+	/** Holds a reference to the message transport technology. */
 	ITransportMessagesPtr Transport;
 };
diff --git a/Engine/Source/Runtime/Messaging/Public/Common/MessageEndpoint.h b/Engine/Source/Runtime/Messaging/Public/Common/MessageEndpoint.h
--- a/Engine/Source/Runtime/Messaging/Public/Common/MessageEndpoint.h
+++ b/Engine/Source/Runtime/Messaging/Public/Common/MessageEndpoint.h
@@ -78,9 +78,7 @@ public:
 		, Name(InName)
 	{ }
 
-	/**
-	 * Destructor.
-	 */
+	/** Destructor. */
 	~FMessageEndpoint( )
 	{
 		IMessageBusPtr Bus = BusPtr.Pin();
@@ -311,7 +309,6 @@ public:
 	 * Checks whether the inbox is empty.
 	 *
 	 * @return true if the inbox is empty, false otherwise.
-	 *
 	 * @see DisableInbox, EnableInbox, IsInboxEnabled, ProcessInbox, ReceiveFromInbox
 	 */
 	bool IsInboxEmpty( ) const
@@ -795,38 +792,38 @@ protected:
 
 private:
 
-	// Holds the endpoint's identifier.
+	/** Holds the endpoint's identifier. */
 	const FMessageAddress Address;
 
-	// Holds a weak pointer to the message bus.
+	/** Holds a weak pointer to the message bus. */
 	IMessageBusWeakPtr BusPtr;
 
-	// Hold a flag indicating whether this endpoint is active.
+	/** Hold a flag indicating whether this endpoint is active. */
 	bool Enabled;
 
-	// Holds the registered message handlers.
+	/** Holds the registered message handlers. */
 	TArray<IMessageHandlerPtr> Handlers;
 
-	// Holds the endpoint's unique identifier (for debugging purposes).
+	/** Holds the endpoint's unique identifier (for debugging purposes). */
 	const FGuid Id;
 
-	// Holds the endpoint's message inbox.
+	/** Holds the endpoint's message inbox. */
 	TQueue<IMessageContextPtr, EQueueMode::Mpsc> Inbox;
 
-	// Holds a flag indicating whether the inbox is enabled.
+	/** Holds a flag indicating whether the inbox is enabled. */
 	bool InboxEnabled;
 
-	// Holds the endpoint's name (for debugging purposes).
+	/** Holds the endpoint's name (for debugging purposes). */
 	const FName Name;
 
-	// Holds a delegate that is invoked when a message has been received.
+	/** Holds a delegate that is invoked when a message has been received. */
 	FOnMessageEndpointReceiveMessage ReceiveDelegate;
 
-	// Holds the name of the thread on which to receive messages.
+	/** Holds the name of the thread on which to receive messages. */
 	ENamedThreads::Type RecipientThread;
 
 private:
 
-	// Holds a delegate that is invoked in case of messaging errors.
+	/** Holds a delegate that is invoked in case of messaging errors. */
 	FOnMessageEndpointError ErrorDelegate;
 };
diff --git a/Engine/Source/Runtime/Messaging/Public/Common/MessageEndpointBuilder.h b/Engine/Source/Runtime/Messaging/Public/Common/MessageEndpointBuilder.h
--- a/Engine/Source/Runtime/Messaging/Public/Common/MessageEndpointBuilder.h
+++ b/Engine/Source/Runtime/Messaging/Public/Common/MessageEndpointBuilder.h
@@ -192,21 +192,21 @@ public:
 	
 private:
 
-	// Holds a reference to the message bus to attach to.
+	/** Holds a reference to the message bus to attach to. */
 	IMessageBusWeakPtr BusPtr;
 
-	// Holds a flag indicating whether the endpoint should be disabled.
+	/** Holds a flag indicating whether the endpoint should be disabled. */
 	bool Disabled;
 
-	// Holds the collection of message handlers to register.
+	/** Holds the collection of message handlers to register. */
 	TArray<IMessageHandlerPtr> Handlers;
 
-	// Holds a flag indicating whether the inbox should be enabled.
+	/** Holds a flag indicating whether the inbox should be enabled. */
 	bool InboxEnabled;
 
-	// Holds the endpoint's name (for debugging purposes).
+	/** Holds the endpoint's name (for debugging purposes). */
 	FName Name;
 
-	// Holds the name of the thread on which to receive messages.
+	/** Holds the name of the thread on which to receive messages. */
 	ENamedThreads::Type RecipientThread;
 };
diff --git a/Engine/Source/Runtime/Messaging/Public/Interfaces/IAuthorizeMessageRecipients.h b/Engine/Source/Runtime/Messaging/Public/Interfaces/IAuthorizeMessageRecipients.h
--- a/Engine/Source/Runtime/Messaging/Public/Interfaces/IAuthorizeMessageRecipients.h
+++ b/Engine/Source/Runtime/Messaging/Public/Interfaces/IAuthorizeMessageRecipients.h
@@ -63,8 +63,6 @@ public:
 
 public:
 
-	/**
-	 * Virtual destructor.
-	 */
+	/** Virtual destructor. */
 	virtual ~IAuthorizeMessageRecipients( ) { }
 };
diff --git a/Engine/Source/Runtime/Messaging/Public/Interfaces/IInterceptMessages.h b/Engine/Source/Runtime/Messaging/Public/Interfaces/IInterceptMessages.h
--- a/Engine/Source/Runtime/Messaging/Public/Interfaces/IInterceptMessages.h
+++ b/Engine/Source/Runtime/Messaging/Public/Interfaces/IInterceptMessages.h
@@ -27,8 +27,6 @@ public:
 
 public:
 
-	/**
-	 * Virtual destructor.
-	 */
+	/** Virtual destructor. */
 	virtual ~IInterceptMessages( ) { }
 };
diff --git a/Engine/Source/Runtime/Messaging/Public/Interfaces/IMessageAttachment.h b/Engine/Source/Runtime/Messaging/Public/Interfaces/IMessageAttachment.h
--- a/Engine/Source/Runtime/Messaging/Public/Interfaces/IMessageAttachment.h
+++ b/Engine/Source/Runtime/Messaging/Public/Interfaces/IMessageAttachment.h
@@ -40,8 +40,6 @@ public:
 
 public:
 
-	/**
-	 * Virtual destructor.
-	 */
+	/** Virtual destructor. */
 	virtual ~IMessageAttachment( ) { }
 };
diff --git a/Engine/Source/Runtime/Messaging/Public/Interfaces/IMessageBridge.h b/Engine/Source/Runtime/Messaging/Public/Interfaces/IMessageBridge.h
--- a/Engine/Source/Runtime/Messaging/Public/Interfaces/IMessageBridge.h
+++ b/Engine/Source/Runtime/Messaging/Public/Interfaces/IMessageBridge.h
@@ -58,8 +58,6 @@ public:
 
 public:
 
-	/**
-	 * Virtual destructor.
-	 */
+	/** Virtual destructor. */
 	virtual ~IMessageBridge( ) { }
 };
diff --git a/Engine/Source/Runtime/Messaging/Public/Interfaces/IMessageBus.h b/Engine/Source/Runtime/Messaging/Public/Interfaces/IMessageBus.h
--- a/Engine/Source/Runtime/Messaging/Public/Interfaces/IMessageBus.h
+++ b/Engine/Source/Runtime/Messaging/Public/Interfaces/IMessageBus.h
@@ -229,8 +229,6 @@ public:
 
 public:
 
-	/**
-	 * Virtual destructor.
-	 */
+	/** Virtual destructor. */
 	virtual ~IMessageBus( ) { }
 };
diff --git a/Engine/Source/Runtime/Messaging/Public/Interfaces/IMessageContext.h b/Engine/Source/Runtime/Messaging/Public/Interfaces/IMessageContext.h
--- a/Engine/Source/Runtime/Messaging/Public/Interfaces/IMessageContext.h
+++ b/Engine/Source/Runtime/Messaging/Public/Interfaces/IMessageContext.h
@@ -194,8 +194,6 @@ public:
 
 public:
 
-	/**
-	 * Virtual destructor.
-	 */
+	/** Virtual destructor. */
 	virtual ~IMessageContext( ) { }
 };
diff --git a/Engine/Source/Runtime/Messaging/Public/Interfaces/IMessageData.h b/Engine/Source/Runtime/Messaging/Public/Interfaces/IMessageData.h
--- a/Engine/Source/Runtime/Messaging/Public/Interfaces/IMessageData.h
+++ b/Engine/Source/Runtime/Messaging/Public/Interfaces/IMessageData.h
@@ -60,4 +60,9 @@ public:
 	 * @return The delegate.
 	 */
 	virtual FSimpleDelegate& OnStateChanged( ) = 0;
+
+public:
+
+	/** Virtual destructor. */
+	virtual ~IMessageData( ) { }
 };
diff --git a/Engine/Source/Runtime/Messaging/Public/Interfaces/IMessageHandler.h b/Engine/Source/Runtime/Messaging/Public/Interfaces/IMessageHandler.h
--- a/Engine/Source/Runtime/Messaging/Public/Interfaces/IMessageHandler.h
+++ b/Engine/Source/Runtime/Messaging/Public/Interfaces/IMessageHandler.h
@@ -33,9 +33,7 @@ public:
 
 public:
 
-	/**
-	 * Virtual destructor.
-	 */
+	/** Virtual destructor. */
 	virtual ~IMessageHandler( ) { }
 };
 
@@ -94,9 +92,9 @@ public:
 	
 private:
 
-	// Holds a pointer to the object handling the messages.
+	/** Holds a pointer to the object handling the messages. */
 	HandlerType* Handler;
 
-	// Holds a pointer to the actual handler function.
+	/** Holds a pointer to the actual handler function. */
 	typename TMessageHandlerFunc<MessageType, HandlerType>::Type HandlerFunc;
 };
diff --git a/Engine/Source/Runtime/Messaging/Public/Interfaces/IMessageSubscription.h b/Engine/Source/Runtime/Messaging/Public/Interfaces/IMessageSubscription.h
--- a/Engine/Source/Runtime/Messaging/Public/Interfaces/IMessageSubscription.h
+++ b/Engine/Source/Runtime/Messaging/Public/Interfaces/IMessageSubscription.h
@@ -27,11 +27,15 @@ public:
 
 	/**
 	 * Disables the subscription.
+	 *
+	 * @see Enable, IsEnabled
 	 */
 	virtual void Disable( ) = 0;
 
 	/**
 	 * Enables the subscription.
+	 *
+	 * @see Disable, IsEnabled
 	 */
 	virtual void Enable( ) = 0;
 
@@ -39,6 +43,7 @@ public:
 	 * Gets the type of subscribed messages.
 	 *
 	 * @return Message type.
+	 * @see GetScopeRange, GetSubscriber
 	 */
 	virtual FName GetMessageType( ) = 0;
 
@@ -46,6 +51,7 @@ public:
 	 * Gets the range of subscribed message scopes.
 	 *
 	 * @return Message scope range.
+	 * @see GetMessageType, GetSubscriber
 	 */
 	virtual const FMessageScopeRange& GetScopeRange( ) = 0;
 
@@ -53,6 +59,7 @@ public:
 	 * Gets the subscriber.
 	 *
 	 * @return The subscriber.
+	 * @see GetMessageType, GetScopeRange
 	 */
 	virtual const IReceiveMessagesWeakPtr& GetSubscriber( ) = 0;
 
@@ -60,13 +67,12 @@ public:
 	 * Checks whether the subscription is enabled.
 	 *
 	 * @return true if the subscription is enabled, false otherwise.
+	 * @see Enable, Disable
 	 */
 	virtual bool IsEnabled( ) = 0;
 
 public:
 
-	/**
-	 * Virtual destructor.
-	 */
+	/** Virtual destructor. */
 	virtual ~IMessageSubscription( ) { }
 };
diff --git a/Engine/Source/Runtime/Messaging/Public/Interfaces/IMessageTracer.h b/Engine/Source/Runtime/Messaging/Public/Interfaces/IMessageTracer.h
--- a/Engine/Source/Runtime/Messaging/Public/Interfaces/IMessageTracer.h
+++ b/Engine/Source/Runtime/Messaging/Public/Interfaces/IMessageTracer.h
@@ -186,12 +186,14 @@ public:
 	/**
 	 * Breaks message routing.
 	 *
-	 * @see Continue, Step
+	 * @see Continue, IsBreaking, Step
 	 */
 	virtual void Break( ) = 0;
 
 	/**
 	 * Continues message routing from the current breakpoint.
+	 *
+	 * @see Break, Step
 	 */
 	virtual void Continue( ) = 0;
 
@@ -199,6 +201,7 @@ public:
 	 * Checks whether the tracer is currently at a breakpoint.
 	 *
 	 * @return true if at breakpoint, false otherwise.
+	 * @see Break
 	 */
 	virtual bool IsBreaking( ) const = 0;
 
@@ -303,8 +306,6 @@ public:
 
 protected:
 
-	/**
-	 * Hidden destructor.
-	 */
+	/** Hidden destructor. */
 	virtual ~IMessageTracer( ) { }
 };
diff --git a/Engine/Source/Runtime/Messaging/Public/Interfaces/IMessageTracerBreakpoint.h b/Engine/Source/Runtime/Messaging/Public/Interfaces/IMessageTracerBreakpoint.h
--- a/Engine/Source/Runtime/Messaging/Public/Interfaces/IMessageTracerBreakpoint.h
+++ b/Engine/Source/Runtime/Messaging/Public/Interfaces/IMessageTracerBreakpoint.h
@@ -36,8 +36,6 @@ public:
 
 protected:
 
-	/**
-	 * Hidden destructor.
-	 */
+	/** Hidden destructor. */
 	virtual ~IMessageTracerBreakpoint( ) { }
 };
diff --git a/Engine/Source/Runtime/Messaging/Public/Interfaces/IMessagingModule.h b/Engine/Source/Runtime/Messaging/Public/Interfaces/IMessagingModule.h
--- a/Engine/Source/Runtime/Messaging/Public/Interfaces/IMessagingModule.h
+++ b/Engine/Source/Runtime/Messaging/Public/Interfaces/IMessagingModule.h
@@ -72,8 +72,6 @@ public:
 
 public:
 
-	/**
-	 * Virtual destructor.
-	 */
+	/** Virtual destructor. */
 	virtual ~IMessagingModule( ) { }
 };
diff --git a/Engine/Source/Runtime/Messaging/Public/Interfaces/IMutableMessageContext.h b/Engine/Source/Runtime/Messaging/Public/Interfaces/IMutableMessageContext.h
--- a/Engine/Source/Runtime/Messaging/Public/Interfaces/IMutableMessageContext.h
+++ b/Engine/Source/Runtime/Messaging/Public/Interfaces/IMutableMessageContext.h
@@ -78,8 +78,6 @@ public:
 
 public:
 
-	/**
-	 * Virtual destructor.
-	 */
+	/** Virtual destructor. */
 	virtual ~IMutableMessageContext( ) { }
 };
diff --git a/Engine/Source/Runtime/Messaging/Public/Interfaces/IReceiveMessages.h b/Engine/Source/Runtime/Messaging/Public/Interfaces/IReceiveMessages.h
--- a/Engine/Source/Runtime/Messaging/Public/Interfaces/IReceiveMessages.h
+++ b/Engine/Source/Runtime/Messaging/Public/Interfaces/IReceiveMessages.h
@@ -91,8 +91,6 @@ public:
 
 public:
 
-	/**
-	 * Virtual destructor.
-	 */
+	/** Virtual destructor. */
 	virtual ~IReceiveMessages( ) { }
 };
diff --git a/Engine/Source/Runtime/Messaging/Public/Interfaces/ISendMessages.h b/Engine/Source/Runtime/Messaging/Public/Interfaces/ISendMessages.h
--- a/Engine/Source/Runtime/Messaging/Public/Interfaces/ISendMessages.h
+++ b/Engine/Source/Runtime/Messaging/Public/Interfaces/ISendMessages.h
@@ -47,8 +47,6 @@ public:
 
 public:
 
-	/**
-	 * Virtual destructor.
-	 */
+	/** Virtual destructor. */
 	virtual ~ISendMessages( ) { }
 };
diff --git a/Engine/Source/Runtime/Messaging/Public/Interfaces/ISerializeMessages.h b/Engine/Source/Runtime/Messaging/Public/Interfaces/ISerializeMessages.h
--- a/Engine/Source/Runtime/Messaging/Public/Interfaces/ISerializeMessages.h
+++ b/Engine/Source/Runtime/Messaging/Public/Interfaces/ISerializeMessages.h
@@ -40,8 +40,6 @@ public:
 
 public:
 
-	/**
-	 * Virtual destructor.
-	 */
+	/** Virtual destructor. */
 	virtual ~ISerializeMessages( ) { }
 };
diff --git a/Engine/Source/Runtime/Messaging/Public/Interfaces/ITransportMessages.h b/Engine/Source/Runtime/Messaging/Public/Interfaces/ITransportMessages.h
--- a/Engine/Source/Runtime/Messaging/Public/Interfaces/ITransportMessages.h
+++ b/Engine/Source/Runtime/Messaging/Public/Interfaces/ITransportMessages.h
@@ -104,8 +104,6 @@ public:
 
 protected:
 
-	/**
-	 * Virtual constructor.
-	 */
+	/** Virtual constructor. */
 	~ITransportMessages( ) { }
 };
diff --git a/Engine/Source/Runtime/Messaging/Public/Messaging.h b/Engine/Source/Runtime/Messaging/Public/Messaging.h
--- a/Engine/Source/Runtime/Messaging/Public/Messaging.h
+++ b/Engine/Source/Runtime/Messaging/Public/Messaging.h
@@ -33,7 +33,7 @@
 #include "IMessagingModule.h"
 
 
-/* Common
+/* Common helpers
  *****************************************************************************/
 
 #include "FileMessageAttachment.h"
