Messaging: Documentation updates

diff --git a/Engine/Source/Runtime/Messaging/Public/Common/FileMessageAttachment.h b/Engine/Source/Runtime/Messaging/Public/Common/FileMessageAttachment.h
--- a/Engine/Source/Runtime/Messaging/Public/Common/FileMessageAttachment.h
+++ b/Engine/Source/Runtime/Messaging/Public/Common/FileMessageAttachment.h
@@ -5,6 +5,8 @@
 
 /**
  * Implements a message attachment whose data is held in a file.
+ *
+ * WARNING: Message attachments do not work yet for out of process messages.
  */
 class FFileMessageAttachment
 	: public IMessageAttachment
diff --git a/Engine/Source/Runtime/Messaging/Public/Common/MessageEndpoint.h b/Engine/Source/Runtime/Messaging/Public/Common/MessageEndpoint.h
--- a/Engine/Source/Runtime/Messaging/Public/Common/MessageEndpoint.h
+++ b/Engine/Source/Runtime/Messaging/Public/Common/MessageEndpoint.h
@@ -13,7 +13,7 @@ typedef TSharedRef<class FMessageEndpoint, ESPMode::ThreadSafe> FMessageEndpoint
 /**
  * DEPRECATED: Delegate type for error notifications.
  *
- * The first parameter is the  context of the sent message.
+ * The first parameter is the context of the sent message.
  * The second parameter is the error string.
  */
 DECLARE_DELEGATE_TwoParams(FOnMessageEndpointError, const IMessageContextRef&, const FString&);
@@ -32,7 +32,7 @@ DECLARE_DELEGATE_RetVal_OneParam(bool, FOnMessageEndpointReceiveMessage, const I
  *
  * This class provides a convenient implementation of the IRecieveMessages and ISendMessages interfaces,
  * which allow consumers to send and receive messages on a message bus. The endpoint allows for receiving
- * messages synchronously as they arrive, as well as asynchronously through an inbox that can be polled.
+ * messages asynchronously as they arrive, as well as synchronously through an inbox that can be polled.
  *
  * By default, messages are received synchronously on the thread that the endpoint was created on.
  * If the message consumer is thread-safe, a more efficient message dispatch can be enabled by calling
@@ -41,8 +41,6 @@ DECLARE_DELEGATE_RetVal_OneParam(bool, FOnMessageEndpointReceiveMessage, const I
  * Endpoints that are destroyed or receive messages on non-Game threads should use the static function
  * FMessageEndpoint::SafeRelease() to dispose of the endpoint. This will ensure that there are no race
  * conditions between endpoint destruction and the receiving of messages.
- *
- * @todo gmp: Messaging reloadability; figure out how to auto-reconnect endpoints
  */
 class FMessageEndpoint
 	: public TSharedFromThis<FMessageEndpoint, ESPMode::ThreadSafe>
diff --git a/Engine/Source/Runtime/Messaging/Public/Common/MessageEndpointBuilder.h b/Engine/Source/Runtime/Messaging/Public/Common/MessageEndpointBuilder.h
--- a/Engine/Source/Runtime/Messaging/Public/Common/MessageEndpointBuilder.h
+++ b/Engine/Source/Runtime/Messaging/Public/Common/MessageEndpointBuilder.h
@@ -13,7 +13,7 @@ public:
 	/**
 	 * Creates and initializes a new builder using the default message bus.
 	 *
-	 * WARNING: This constructor must be called from the main thread.
+	 * WARNING: This constructor must be called from the Game thread.
 	 *
 	 * @param InName The endpoint's name (for debugging purposes).
 	 * @param InBus The message bus to attach the endpoint to.
diff --git a/Engine/Source/Runtime/Messaging/Public/Interfaces/IMessageTracer.h b/Engine/Source/Runtime/Messaging/Public/Interfaces/IMessageTracer.h
--- a/Engine/Source/Runtime/Messaging/Public/Interfaces/IMessageTracer.h
+++ b/Engine/Source/Runtime/Messaging/Public/Interfaces/IMessageTracer.h
@@ -176,7 +176,7 @@ DECLARE_MULTICAST_DELEGATE_OneParam(FMessageTracerTypeAdded, FMessageTracerTypeI
  * Interface for message tracers.
  *
  * This interface provides access to the message tracer that is built into a message bus. Message tracers are
- * able to monitor and record all activities on a message bus. This includes the sending and receiving messages,
+ * able to monitor and record all activities on a message bus. This includes the sending and receiving of messages,
  * the registering and unregistering of message endpoints and interceptors and adding and removal of subscriptions.
  */
 class IMessageTracer
@@ -306,6 +306,6 @@ public:
 
 protected:
 
-	/** Hidden destructor. */
+	/** Hidden destructor. The life time of a message tracer is managed by the message bus. */
 	virtual ~IMessageTracer( ) { }
 };
