Merging using UE4-Releases-4.5-To-UE4 for JsonMessageSerializer.cpp
Fixed arrays of structs not desirializing properly.

diff --git a/Engine/Source/Runtime/Messaging/Private/Serialization/JsonMessageSerializer.cpp b/Engine/Source/Runtime/Messaging/Private/Serialization/JsonMessageSerializer.cpp
--- a/Engine/Source/Runtime/Messaging/Private/Serialization/JsonMessageSerializer.cpp
+++ b/Engine/Source/Runtime/Messaging/Private/Serialization/JsonMessageSerializer.cpp
@@ -433,7 +433,7 @@ bool FJsonMessageSerializer::DeserializeStruct( IMutableMessageContextRef& OutCo
 
 						FScriptArrayHelper ArrayHelper(ArrayProperty, ArrayProperty->ContainerPtrToValuePtr<void>(CurrentState.Data));
 						const int32 ArrayIndex = ArrayHelper.AddValue();
-						NewState.Property = CurrentState.Property;
+						NewState.Property = ArrayProperty->Inner;
 						NewState.Data = ArrayHelper.GetRawPtr(ArrayIndex);
 					}
 					else
