Messaging: Documentation fix & cosmetic changes

diff --git a/Engine/Plugins/Messaging/UdpMessaging/Source/UdpMessaging/Private/UdpMessagingModule.cpp b/Engine/Plugins/Messaging/UdpMessaging/Source/UdpMessaging/Private/UdpMessagingModule.cpp
--- a/Engine/Plugins/Messaging/UdpMessaging/Source/UdpMessaging/Private/UdpMessagingModule.cpp
+++ b/Engine/Plugins/Messaging/UdpMessaging/Source/UdpMessaging/Private/UdpMessagingModule.cpp
@@ -225,7 +225,7 @@ protected:
 	{
 		if (FApp::IsGame())
 		{
-			// disallow in Debug and Development configurations
+			// only allow in Debug and Development configurations for now
 			if ((FApp::GetBuildConfiguration() == EBuildConfigurations::Shipping) || (FApp::GetBuildConfiguration() == EBuildConfigurations::Test))
 			{
 				return false;
diff --git a/Engine/Plugins/Messaging/UdpMessaging/Source/UdpMessaging/Private/UdpMessagingPrivatePCH.h b/Engine/Plugins/Messaging/UdpMessaging/Source/UdpMessaging/Private/UdpMessagingPrivatePCH.h
--- a/Engine/Plugins/Messaging/UdpMessaging/Source/UdpMessaging/Private/UdpMessagingPrivatePCH.h
+++ b/Engine/Plugins/Messaging/UdpMessaging/Source/UdpMessaging/Private/UdpMessagingPrivatePCH.h
@@ -5,7 +5,7 @@
 #include "UdpMessaging.h"
 
 
-/* Private Dependencies
+/* Private dependencies
  *****************************************************************************/
 
 #include "Core.h"
@@ -17,7 +17,7 @@
 #include "SocketSubsystem.h"
 
 
-/* Private Constants
+/* Private constants
  *****************************************************************************/
 
 /** Defines the default IP endpoint for multicast traffic. */
diff --git a/Engine/Source/Runtime/Messaging/Private/Bridge/MessageBridge.h b/Engine/Source/Runtime/Messaging/Private/Bridge/MessageBridge.h
--- a/Engine/Source/Runtime/Messaging/Private/Bridge/MessageBridge.h
+++ b/Engine/Source/Runtime/Messaging/Private/Bridge/MessageBridge.h
@@ -78,20 +78,18 @@ public:
 
 protected:
 
-	/**
-	 * Shuts down the bridge.
-	 */
+	/** Shuts down the bridge. */
 	void Shutdown( );
 
 private:
 
-	// Callback for message bus shut downs
+	/** Callback for message bus shut downs. */
 	void HandleMessageBusShutdown( );
 
-	// Callback for received transport messages.
+	/** Callback for received transport messages. */
 	void HandleTransportMessageReceived( FArchive& MessageData, const IMessageAttachmentPtr& Attachment, const FGuid& NodeId );
 
-	// Callback for lost transport endpoints.
+	/** Callback for lost transport endpoints. */
 	void HandleTransportNodeLost( const FGuid& LostNodeId );
 
 private:
diff --git a/Engine/Source/Runtime/Messaging/Private/Bridge/MessageData.h b/Engine/Source/Runtime/Messaging/Private/Bridge/MessageData.h
--- a/Engine/Source/Runtime/Messaging/Private/Bridge/MessageData.h
+++ b/Engine/Source/Runtime/Messaging/Private/Bridge/MessageData.h
@@ -59,10 +59,10 @@ public:
 
 private:
 
-	// Holds the data.
+	/** Holds the data. */
 	TArray<uint8> Data;
 
-	// Holds the message data state.
+	/** Holds the message data state. */
 	EMessageDataState::Type State;
 
 private:
