Messaging: Replaced namespaced enums with enum classes

diff --git a/Engine/Plugins/Messaging/MessagingDebugger/Source/MessagingDebugger/Private/Widgets/History/SMessagingHistoryTableRow.h b/Engine/Plugins/Messaging/MessagingDebugger/Source/MessagingDebugger/Private/Widgets/History/SMessagingHistoryTableRow.h
--- a/Engine/Plugins/Messaging/MessagingDebugger/Source/MessagingDebugger/Private/Widgets/History/SMessagingHistoryTableRow.h
+++ b/Engine/Plugins/Messaging/MessagingDebugger/Source/MessagingDebugger/Private/Widgets/History/SMessagingHistoryTableRow.h
@@ -223,7 +223,7 @@ protected:
 	 *
 	 * @return The string representation.
 	 */
-	FText ScopeToText( EMessageScope::Type Scope ) const
+	FText ScopeToText( EMessageScope Scope ) const
 	{
 		switch (Scope)
 		{
diff --git a/Engine/Source/Runtime/Messaging/Private/Bridge/MessageData.h b/Engine/Source/Runtime/Messaging/Private/Bridge/MessageData.h
--- a/Engine/Source/Runtime/Messaging/Private/Bridge/MessageData.h
+++ b/Engine/Source/Runtime/Messaging/Private/Bridge/MessageData.h
@@ -20,7 +20,7 @@ class FMessageData
 public:
 
 	/** Default constructor. */
-	FMessageData( )
+	FMessageData()
 		: FMemoryWriter(Data, true)
 		, State(EMessageDataState::Incomplete)
 	{ }
@@ -32,7 +32,7 @@ public:
 	 *
 	 * @param InState The state to set.
 	 */
-	void UpdateState( EMessageDataState::Type InState )
+	void UpdateState( EMessageDataState InState )
 	{
 		State = InState;
 		StateChangedDelegate.ExecuteIfBound();
@@ -42,17 +42,17 @@ public:
 
 	// IMessageAttachment interface
 
-	virtual FArchive* CreateReader( ) override
+	virtual FArchive* CreateReader() override
 	{
 		return new FMemoryReader(Data, true);
 	}
 
-	virtual EMessageDataState::Type GetState( ) const override
+	virtual EMessageDataState GetState() const override
 	{
 		return State;
 	}
 
-	virtual FSimpleDelegate& OnStateChanged( ) override
+	virtual FSimpleDelegate& OnStateChanged() override
 	{
 		return StateChangedDelegate;
 	}
@@ -63,7 +63,7 @@ private:
 	TArray<uint8> Data;
 
 	/** Holds the message data state. */
-	EMessageDataState::Type State;
+	EMessageDataState State;
 
 private:
 
diff --git a/Engine/Source/Runtime/Messaging/Private/Bus/MessageBus.h b/Engine/Source/Runtime/Messaging/Private/Bus/MessageBus.h
--- a/Engine/Source/Runtime/Messaging/Private/Bus/MessageBus.h
+++ b/Engine/Source/Runtime/Messaging/Private/Bus/MessageBus.h
@@ -20,18 +20,18 @@ public:
 	FMessageBus( const IAuthorizeMessageRecipientsPtr& InRecipientAuthorizer );
 
 	/** Destructor */
-	~FMessageBus( );
+	~FMessageBus();
 
 public:
 
 	// IMessageBus interface
 
-	virtual void Forward( const IMessageContextRef& Context, const TArray<FMessageAddress>& Recipients, EMessageScope::Type ForwardingScope, const FTimespan& Delay, const ISendMessagesRef& Forwarder ) override
+	virtual void Forward( const IMessageContextRef& Context, const TArray<FMessageAddress>& Recipients, EMessageScope ForwardingScope, const FTimespan& Delay, const ISendMessagesRef& Forwarder ) override
 	{
 		Router->RouteMessage(MakeShareable(new FMessageContext(Context, Forwarder->GetSenderAddress(), Recipients, ForwardingScope, FDateTime::UtcNow(), FTaskGraphInterface::Get().GetCurrentThreadIfKnown())));
 	}
 
-	virtual IMessageTracerRef GetTracer( ) override
+	virtual IMessageTracerRef GetTracer() override
 	{
 		return Router->GetTracer();
 	}
@@ -47,12 +47,12 @@ public:
 		}
 	}
 
-	virtual FOnMessageBusShutdown& OnShutdown( ) override
+	virtual FOnMessageBusShutdown& OnShutdown() override
 	{
 		return ShutdownDelegate;
 	}
 
-	virtual void Publish( void* Message, UScriptStruct* TypeInfo, EMessageScope::Type Scope, const FTimespan& Delay, const FDateTime& Expiration, const ISendMessagesRef& Publisher ) override
+	virtual void Publish( void* Message, UScriptStruct* TypeInfo, EMessageScope Scope, const FTimespan& Delay, const FDateTime& Expiration, const ISendMessagesRef& Publisher ) override
 	{
 		Router->RouteMessage(MakeShareable(new FMessageContext(Message, TypeInfo, nullptr, Publisher->GetSenderAddress(), TArray<FMessageAddress>(), Scope, FDateTime::UtcNow() + Delay, Expiration, FTaskGraphInterface::Get().GetCurrentThreadIfKnown())));
 	}
@@ -67,7 +67,7 @@ public:
 		Router->RouteMessage(MakeShareable(new FMessageContext(Message, TypeInfo, Attachment, Sender->GetSenderAddress(), Recipients, EMessageScope::Network, FDateTime::UtcNow() + Delay, Expiration, FTaskGraphInterface::Get().GetCurrentThreadIfKnown())));
 	}
 
-	virtual void Shutdown( ) override;
+	virtual void Shutdown() override;
 
 	virtual IMessageSubscriptionPtr Subscribe( const IReceiveMessagesRef& Subscriber, const FName& MessageType, const FMessageScopeRange& ScopeRange ) override
 	{
diff --git a/Engine/Source/Runtime/Messaging/Private/Bus/MessageContext.h b/Engine/Source/Runtime/Messaging/Private/Bus/MessageContext.h
--- a/Engine/Source/Runtime/Messaging/Private/Bus/MessageContext.h
+++ b/Engine/Source/Runtime/Messaging/Private/Bus/MessageContext.h
@@ -15,7 +15,7 @@ class FMessageContext
 public:
 
 	/** Default constructor. */
-	FMessageContext( )
+	FMessageContext()
 		: Message(nullptr)
 		, TypeInfo(nullptr)
 	{ }
@@ -33,7 +33,7 @@ public:
 	 * @param InExpiration The message's expiration time.
 	 * @param InSenderThread The name of the thread from which the message was sent.
 	 */
-	FMessageContext( void* InMessage, UScriptStruct* InTypeInfo, const IMessageAttachmentPtr& InAttachment, const FMessageAddress& InSender, const TArray<FMessageAddress>& InRecipients, EMessageScope::Type InScope, const FDateTime& InTimeSent, const FDateTime& InExpiration, ENamedThreads::Type InSenderThread )
+	FMessageContext( void* InMessage, UScriptStruct* InTypeInfo, const IMessageAttachmentPtr& InAttachment, const FMessageAddress& InSender, const TArray<FMessageAddress>& InRecipients, EMessageScope InScope, const FDateTime& InTimeSent, const FDateTime& InExpiration, ENamedThreads::Type InSenderThread )
 		: Attachment(InAttachment)
 		, Expiration(InExpiration)
 		, Message(InMessage)
@@ -55,7 +55,7 @@ public:
 	 * @param InTimeForwarded The time at which the message was forwarded.
 	 * @param InSenderThread The name of the thread from which the message was sent.
 	 */
-	FMessageContext( const IMessageContextRef& InContext, const FMessageAddress& InForwarder, const TArray<FMessageAddress>& NewRecipients, EMessageScope::Type NewScope, const FDateTime& InTimeForwarded, ENamedThreads::Type InSenderThread )
+	FMessageContext( const IMessageContextRef& InContext, const FMessageAddress& InForwarder, const TArray<FMessageAddress>& NewRecipients, EMessageScope NewScope, const FDateTime& InTimeForwarded, ENamedThreads::Type InSenderThread )
 		: Attachment(nullptr)
 		, Forwarder(InForwarder)
 		, Message(nullptr)
@@ -81,7 +81,7 @@ public:
 
 	// IMessageContext interface
 
-	virtual IMessageAttachmentPtr GetAttachment( ) const override
+	virtual IMessageAttachmentPtr GetAttachment() const override
 	{
 		if (OriginalContext.IsValid())
 		{
@@ -91,7 +91,7 @@ public:
 		return Attachment;
 	}
 
-	virtual const FDateTime& GetExpiration( ) const override
+	virtual const FDateTime& GetExpiration() const override
 	{
 		if (OriginalContext.IsValid())
 		{
@@ -101,12 +101,12 @@ public:
 		return Expiration;
 	}
 
-	virtual const FMessageAddress& GetForwarder( ) const override
+	virtual const FMessageAddress& GetForwarder() const override
 	{
 		return Forwarder;
 	}
 
-	virtual const TMap<FName, FString>& GetHeaders( ) const override
+	virtual const TMap<FName, FString>& GetHeaders() const override
 	{
 		if (OriginalContext.IsValid())
 		{
@@ -116,7 +116,7 @@ public:
 		return Headers;
 	}
 
-	virtual const void* GetMessage( ) const override
+	virtual const void* GetMessage() const override
 	{
 		if (OriginalContext.IsValid())
 		{
@@ -126,7 +126,7 @@ public:
 		return Message;
 	}
 
-	virtual FName GetMessageType( ) const override
+	virtual FName GetMessageType() const override
 	{
 		if (IsValid())
 		{
@@ -136,7 +136,7 @@ public:
 		return NAME_None;
 	}
 
-	virtual const TWeakObjectPtr<UScriptStruct>& GetMessageTypeInfo( ) const override
+	virtual const TWeakObjectPtr<UScriptStruct>& GetMessageTypeInfo() const override
 	{
 		if (OriginalContext.IsValid())
 		{
@@ -146,22 +146,22 @@ public:
 		return TypeInfo;
 	}
 
-	virtual IMessageContextPtr GetOriginalContext( ) const override
+	virtual IMessageContextPtr GetOriginalContext() const override
 	{
 		return OriginalContext;
 	}
 
-	virtual const TArray<FMessageAddress>& GetRecipients( ) const override
+	virtual const TArray<FMessageAddress>& GetRecipients() const override
 	{
 		return Recipients;
 	}
 
-	virtual EMessageScope::Type GetScope( ) const override
+	virtual EMessageScope GetScope() const override
 	{
 		return Scope;
 	}
 
-	virtual const FMessageAddress& GetSender( ) const override
+	virtual const FMessageAddress& GetSender() const override
 	{
 		if (OriginalContext.IsValid())
 		{
@@ -171,17 +171,17 @@ public:
 		return Sender;
 	}
 
-	virtual ENamedThreads::Type GetSenderThread( ) const override
+	virtual ENamedThreads::Type GetSenderThread() const override
 	{
 		return SenderThread;
 	}
 
-	virtual const FDateTime& GetTimeForwarded( ) const override
+	virtual const FDateTime& GetTimeForwarded() const override
 	{
 		return TimeForwarded;
 	}
 
-	virtual const FDateTime& GetTimeSent( ) const override
+	virtual const FDateTime& GetTimeSent() const override
 	{
 		if (OriginalContext.IsValid())
 		{
@@ -191,12 +191,12 @@ public:
 		return TimeSent;
 	}
 
-	virtual bool IsForwarded( ) const override
+	virtual bool IsForwarded() const override
 	{
 		return OriginalContext.IsValid();
 	}
 
-	virtual bool IsValid( ) const override
+	virtual bool IsValid() const override
 	{
 		if (OriginalContext.IsValid())
 		{
@@ -236,7 +236,7 @@ public:
 		Headers.Add(Key, Value);
 	}
 
-	virtual void SetScope( EMessageScope::Type InScope ) override
+	virtual void SetScope( EMessageScope InScope ) override
 	{
 		Scope = InScope;
 	}
@@ -275,7 +275,7 @@ private:
 	TArray<FMessageAddress> Recipients;
 
 	/** Holds the message's scope. */
-	EMessageScope::Type Scope;
+	EMessageScope Scope;
 
 	/** Holds the sender's identifier. */
 	FMessageAddress Sender;
diff --git a/Engine/Source/Runtime/Messaging/Private/Bus/MessageRouter.cpp b/Engine/Source/Runtime/Messaging/Private/Bus/MessageRouter.cpp
--- a/Engine/Source/Runtime/Messaging/Private/Bus/MessageRouter.cpp
+++ b/Engine/Source/Runtime/Messaging/Private/Bus/MessageRouter.cpp
@@ -164,7 +164,7 @@ void FMessageRouter::DispatchMessage( const IMessageContextRef& Context )
 
 void FMessageRouter::FilterSubscriptions( TArray<IMessageSubscriptionPtr>& Subscriptions, const IMessageContextRef& Context, TArray<IReceiveMessagesPtr>& OutRecipients )
 {
-	EMessageScope::Type MessageScope = Context->GetScope();
+	EMessageScope MessageScope = Context->GetScope();
 
 	for (int32 SubscriptionIndex = 0; SubscriptionIndex < Subscriptions.Num(); ++SubscriptionIndex)
 	{
diff --git a/Engine/Source/Runtime/Messaging/Private/Serialization/JsonMessageSerializer.cpp b/Engine/Source/Runtime/Messaging/Private/Serialization/JsonMessageSerializer.cpp
--- a/Engine/Source/Runtime/Messaging/Private/Serialization/JsonMessageSerializer.cpp
+++ b/Engine/Source/Runtime/Messaging/Private/Serialization/JsonMessageSerializer.cpp
@@ -206,7 +206,7 @@ bool FJsonMessageSerializer::DeserializeMessage( FArchive& Archive, IMutableMess
 		OutContext->AddRecipient(Recipients[RecipientIndex]);
 	}
 
-	TEnumAsByte<EMessageScope::Type> Scope;
+	TEnumAsByte<EMessageScope> Scope;
 	Archive << Scope;
 	OutContext->SetScope(Scope);
 
@@ -584,7 +584,7 @@ bool FJsonMessageSerializer::SerializeMessage( const IMessageContextRef& Context
 		TArray<FMessageAddress> Recipients = Context->GetRecipients();
 		Archive << Recipients;
 
-		TEnumAsByte<EMessageScope::Type> Scope = Context->GetScope();
+		TEnumAsByte<EMessageScope> Scope = Context->GetScope();
 		Archive << Scope;
 
 		FDateTime TimeSent = Context->GetTimeSent();
diff --git a/Engine/Source/Runtime/Messaging/Public/Common/MessageEndpoint.h b/Engine/Source/Runtime/Messaging/Public/Common/MessageEndpoint.h
--- a/Engine/Source/Runtime/Messaging/Public/Common/MessageEndpoint.h
+++ b/Engine/Source/Runtime/Messaging/Public/Common/MessageEndpoint.h
@@ -77,7 +77,7 @@ public:
 	{ }
 
 	/** Destructor. */
-	~FMessageEndpoint( )
+	~FMessageEndpoint()
 	{
 		IMessageBusPtr Bus = BusPtr.Pin();
 
@@ -97,7 +97,7 @@ public:
 	 *
 	 * @see Enable
 	 */
-	void Disable( )
+	void Disable()
 	{
 		Enabled = false;
 	}
@@ -110,7 +110,7 @@ public:
 	 *
 	 * @see Disable
 	 */
-	void Enable( )
+	void Enable()
 	{
 		Enabled = true;
 	}
@@ -120,7 +120,7 @@ public:
 	 *
 	 * @return Message address.
 	 */
-	const FMessageAddress& GetAddress( ) const
+	const FMessageAddress& GetAddress() const
 	{
 		return Address;
 	}
@@ -128,7 +128,7 @@ public:
 	/**
 	 * Checks whether this endpoint is connected to the bus.
 	 */
-	bool IsConnected( ) const
+	bool IsConnected() const
 	{
 		return BusPtr.IsValid();
 	}
@@ -138,7 +138,7 @@ public:
 	 *
 	 * @return true if the endpoint is enabled, false otherwise.
 	 */
-	bool IsEnabled( ) const
+	bool IsEnabled() const
 	{
 		return Enabled;
 	}
@@ -193,7 +193,7 @@ public:
 	 * @param ForwardingScope The scope of the forwarded message.
 	 * @param Delay The time delay.
 	 */
-	void Forward( const IMessageContextRef& Context, const TArray<FMessageAddress>& Recipients, EMessageScope::Type ForwardingScope, const FTimespan& Delay )
+	void Forward( const IMessageContextRef& Context, const TArray<FMessageAddress>& Recipients, EMessageScope ForwardingScope, const FTimespan& Delay )
 	{
 		IMessageBusPtr Bus = GetBusIfEnabled();
 
@@ -213,7 +213,7 @@ public:
 	 * @param Delay The delay after which to publish the message.
 	 * @param Expiration The time at which the message expires.
 	 */
-	void Publish( void* Message, UScriptStruct* TypeInfo, EMessageScope::Type Scope, const FTimespan& Delay, const FDateTime& Expiration )
+	void Publish( void* Message, UScriptStruct* TypeInfo, EMessageScope Scope, const FTimespan& Delay, const FDateTime& Expiration )
 	{
 		IMessageBusPtr Bus = GetBusIfEnabled();
 
@@ -284,7 +284,7 @@ public:
 	 *
 	 * @see EnableInbox, IsInboxEmpty, IsInboxEnabled, ProcessInbox, ReceiveFromInbox
 	 */
-	void DisableInbox( )
+	void DisableInbox()
 	{
 		InboxEnabled = false;
 	}
@@ -298,7 +298,7 @@ public:
 	 *
 	 * @see DisableInbox, IsInboxEmpty, IsInboxEnabled, ProcessInbox, ReceiveFromInbox
 	 */
-	void EnableInbox( )
+	void EnableInbox()
 	{
 		InboxEnabled = true;
 	}
@@ -309,7 +309,7 @@ public:
 	 * @return true if the inbox is empty, false otherwise.
 	 * @see DisableInbox, EnableInbox, IsInboxEnabled, ProcessInbox, ReceiveFromInbox
 	 */
-	bool IsInboxEmpty( ) const
+	bool IsInboxEmpty() const
 	{
 		return Inbox.IsEmpty();
 	}
@@ -319,7 +319,7 @@ public:
 	 *
 	 * @see DisableInbox, EnableInbox, IsInboxEmpty, ProcessInbox, ReceiveFromInbox
 	 */
-	bool IsInboxEnabled( ) const
+	bool IsInboxEnabled() const
 	{
 		return InboxEnabled;
 	}
@@ -333,7 +333,7 @@ public:
 	 *
 	 * @see IsInboxEmpty, ReceiveFromInbox
 	 */
-	void ProcessInbox( )
+	void ProcessInbox()
 	{
 		IMessageContextPtr Context;
 
@@ -369,7 +369,7 @@ public:
 	 *
 	 * @return The delegate.
 	 */
-	FOnMessageEndpointReceiveMessage& OnReceiveMessage( )
+	FOnMessageEndpointReceiveMessage& OnReceiveMessage()
 	{
 		return ReceiveDelegate;
 	}
@@ -378,22 +378,22 @@ public:
 
 	// IReceiveMessages interface
 
-	virtual FName GetDebugName( ) const override
+	virtual FName GetDebugName() const override
 	{
 		return Name;
 	}
 
-	virtual const FGuid& GetRecipientId( ) const override
+	virtual const FGuid& GetRecipientId() const override
 	{
 		return Id;
 	}
 
-	virtual ENamedThreads::Type GetRecipientThread( ) const override
+	virtual ENamedThreads::Type GetRecipientThread() const override
 	{
 		return RecipientThread;
 	}
 
-	virtual bool IsLocal( ) const override
+	virtual bool IsLocal() const override
 	{
 		return true;
 	}
@@ -424,7 +424,7 @@ public:
 
 	// ISendMessages interface
 
-	virtual FMessageAddress GetSenderAddress( ) override
+	virtual FMessageAddress GetSenderAddress() override
 	{
 		return Address;
 	}
@@ -445,7 +445,7 @@ public:
 	 * @param Recipient The address of the recipient to forward the message to.
 	 * @param ForwardingScope The scope of the forwarded message.
 	 */
-	void Forward( const IMessageContextRef& Context, const FMessageAddress& Recipient, EMessageScope::Type ForwardingScope )
+	void Forward( const IMessageContextRef& Context, const FMessageAddress& Recipient, EMessageScope ForwardingScope )
 	{
 		Forward(Context, TArrayBuilder<FMessageAddress>().Add(Recipient), ForwardingScope, FTimespan::Zero());
 	}
@@ -460,7 +460,7 @@ public:
 	 * @param ForwardingScope The scope of the forwarded message.
 	 * @param Delay The delay after which to publish the message.
 	 */
-	void Forward( const IMessageContextRef& Context, const FMessageAddress& Recipient, EMessageScope::Type ForwardingScope, const FTimespan& Delay )
+	void Forward( const IMessageContextRef& Context, const FMessageAddress& Recipient, EMessageScope ForwardingScope, const FTimespan& Delay )
 	{
 		Forward(Context, TArrayBuilder<FMessageAddress>().Add(Recipient), ForwardingScope, Delay);
 	}
@@ -474,7 +474,7 @@ public:
 	 * @param Recipients The list of message recipients to forward the message to.
 	 * @param ForwardingScope The scope of the forwarded message.
 	 */
-	void Forward( const IMessageContextRef& Context, const TArray<FMessageAddress>& Recipients, EMessageScope::Type ForwardingScope )
+	void Forward( const IMessageContextRef& Context, const TArray<FMessageAddress>& Recipients, EMessageScope ForwardingScope )
 	{
 		Forward(Context, Recipients, ForwardingScope, FTimespan::Zero());
 	}
@@ -497,7 +497,7 @@ public:
 	 * @param Scope The message scope.
 	 */
 	template<typename MessageType>
-	void Publish( MessageType* Message, EMessageScope::Type Scope )
+	void Publish( MessageType* Message, EMessageScope Scope )
 	{
 		Publish(Message, MessageType::StaticStruct(), Scope, FTimespan::Zero(), FDateTime::MaxValue());
 	}
@@ -522,7 +522,7 @@ public:
 	 * @param Delay The delay after which to publish the message.
 	 */
 	template<typename MessageType>
-	void Publish( MessageType* Message, EMessageScope::Type Scope, const FTimespan& Delay )
+	void Publish( MessageType* Message, EMessageScope Scope, const FTimespan& Delay )
 	{
 		Publish(Message, MessageType::StaticStruct(), Scope, Delay, FDateTime::MaxValue());
 	}
@@ -537,7 +537,7 @@ public:
 	 * @param Expiration The time at which the message expires.
 	 */
 	template<typename MessageType>
-	void Publish( MessageType* Message, EMessageScope::Type Scope, const FTimespan& Delay, const FDateTime& Expiration )
+	void Publish( MessageType* Message, EMessageScope Scope, const FTimespan& Delay, const FDateTime& Expiration )
 	{
 		Publish(Message, MessageType::StaticStruct(), Scope, Delay, Expiration);
 	}
@@ -683,7 +683,7 @@ public:
 	 * @param HandlerFunc The class function handling the messages.
 	 */
 	template<class MessageType>
-	void Subscribe( )
+	void Subscribe()
 	{
 		Subscribe(MessageType::StaticStruct()->GetFName(), FMessageScopeRange::AtLeast(EMessageScope::Thread));
 	}
@@ -708,7 +708,7 @@ public:
 	 *
 	 * @see Subscribe
 	 */
-	void Unsubscribe( )
+	void Unsubscribe()
 	{
 		Unsubscribe(NAME_All);
 	}
@@ -720,7 +720,7 @@ public:
 	 * @see Subscribe
 	 */
 	template<class MessageType>
-	void Unsubscribe( )
+	void Unsubscribe()
 	{
 		Unsubscribe(MessageType::StaticStruct()->GetFName());
 	}
@@ -755,7 +755,7 @@ protected:
 	 *
 	 * @return The message bus.
 	 */
-	FORCEINLINE IMessageBusPtr GetBusIfEnabled( ) const
+	FORCEINLINE IMessageBusPtr GetBusIfEnabled() const
 	{
 		if (Enabled)
 		{
diff --git a/Engine/Source/Runtime/Messaging/Public/Interfaces/IMessageBus.h b/Engine/Source/Runtime/Messaging/Public/Interfaces/IMessageBus.h
--- a/Engine/Source/Runtime/Messaging/Public/Interfaces/IMessageBus.h
+++ b/Engine/Source/Runtime/Messaging/Public/Interfaces/IMessageBus.h
@@ -114,14 +114,14 @@ public:
 	 * @param Forwarder The sender that forwards the message.
 	 * @see Publish, Send
 	 */
-	virtual void Forward( const IMessageContextRef& Context, const TArray<FMessageAddress>& Recipients, EMessageScope::Type NewScope, const FTimespan& Delay, const ISendMessagesRef& Forwarder ) = 0;
+	virtual void Forward( const IMessageContextRef& Context, const TArray<FMessageAddress>& Recipients, EMessageScope NewScope, const FTimespan& Delay, const ISendMessagesRef& Forwarder ) = 0;
 
 	/**
 	 * Gets the message bus tracer.
 	 *
 	 * @return Weak pointer to the tracer object.
 	 */
-	virtual IMessageTracerRef GetTracer( ) = 0;
+	virtual IMessageTracerRef GetTracer() = 0;
 
 	/**
 	 * Adds an interceptor for messages of the specified type.
@@ -145,7 +145,7 @@ public:
 	 * @param Publisher The message publisher.
 	 * @see Forward, Send
 	 */
-	virtual void Publish( void* Message, UScriptStruct* TypeInfo, EMessageScope::Type Scope, const FTimespan& Delay, const FDateTime& Expiration, const ISendMessagesRef& Publisher ) = 0;
+	virtual void Publish( void* Message, UScriptStruct* TypeInfo, EMessageScope Scope, const FTimespan& Delay, const FDateTime& Expiration, const ISendMessagesRef& Publisher ) = 0;
 
 	/**
 	 * Registers a message recipient with the message bus.
@@ -175,7 +175,7 @@ public:
 	 *
 	 * @see OnShutdown
 	 */
-	virtual void Shutdown( ) = 0;
+	virtual void Shutdown() = 0;
 
 	/**
 	 * Adds a subscription for published messages of the specified type.
@@ -225,10 +225,10 @@ public:
 	 * @return The delegate.
 	 * @see Shutdown
 	 */
-	virtual FOnMessageBusShutdown& OnShutdown( ) = 0;
+	virtual FOnMessageBusShutdown& OnShutdown() = 0;
 
 public:
 
 	/** Virtual destructor. */
-	virtual ~IMessageBus( ) { }
+	virtual ~IMessageBus() { }
 };
diff --git a/Engine/Source/Runtime/Messaging/Public/Interfaces/IMessageContext.h b/Engine/Source/Runtime/Messaging/Public/Interfaces/IMessageContext.h
--- a/Engine/Source/Runtime/Messaging/Public/Interfaces/IMessageContext.h
+++ b/Engine/Source/Runtime/Messaging/Public/Interfaces/IMessageContext.h
@@ -3,42 +3,39 @@
 #pragma once
 
 
-namespace EMessageScope
+/**
+ * Enumerates scopes for published messages.
+ *
+ * The scope determines to which endpoints a published message will be delivered.
+ * By default, messages will be published to everyone on the local network (Subnet),
+ * but it is often useful to restrict the group of recipients to more local scopes,
+ * or to widen it to a larger audience outside the local network.
+ *
+ * For example, if a message is to be handled only by subscribers in the same application,
+ * the scope should be set to Process. If messages need to be published between
+ * different networks (i.e. between LAN and WLAN), it should be set to Network instead.
+ * 
+ * Scopes only apply to published messages. Messages that are being sent to specific
+ * recipients will always be delivered, regardless of the endpoint locations.
+ */
+enum class EMessageScope
 {
-	/**
-	 * Enumerates scopes for published messages.
-	 *
-	 * The scope determines to which endpoints a published message will be delivered.
-	 * By default, messages will be published to everyone on the local network (Subnet),
-	 * but it is often useful to restrict the group of recipients to more local scopes,
-	 * or to widen it to a larger audience outside the local network.
-	 *
-	 * For example, if a message is to be handled only by subscribers in the same application,
-	 * the scope should be set to Process. If messages need to be published between
-	 * different networks (i.e. between LAN and WLAN), it should be set to Network instead.
-	 *
-	 * Scopes only apply to published messages. Messages that are being sent to specific
-	 * recipients will always be delivered, regardless of the endpoint locations.
-	 */
-	enum Type
-	{
-		/** Deliver to subscribers in the same thread. */
-		Thread,
+	/** Deliver to subscribers in the same thread. */
+	Thread,
 
-		/** Deliver to subscribers in the same process. */
-		Process,
+	/** Deliver to subscribers in the same process. */
+	Process,
 
-		/** Deliver to subscribers on the network. */
-		Network,
+	/** Deliver to subscribers on the network. */
+	Network,
 
-		/**
-		 * Deliver to all subscribers.
-		 *
-		 * Note: This must be the last value in this enumeration.
-		 */
-		All
-	};
-}
+	/**
+	 * Deliver to all subscribers.
+	 *
+	 * Note: This must be the last value in this enumeration.
+	 */
+	All
+};
 
 
 /** Type definition for message endpoint identifiers. */
@@ -51,10 +48,10 @@ typedef TSharedPtr<class IMessageContext, ESPMode::ThreadSafe> IMessageContextPt
 typedef TSharedRef<class IMessageContext, ESPMode::ThreadSafe> IMessageContextRef;
 
 /** Type definition for message scope ranges. */
-typedef TRange<EMessageScope::Type> FMessageScopeRange;
+typedef TRange<EMessageScope> FMessageScopeRange;
 
 /** Type definition for message scope range bounds. */
-typedef TRangeBound<EMessageScope::Type> FMessageScopeRangeBound;
+typedef TRangeBound<EMessageScope> FMessageScopeRangeBound;
 
 
 /**
@@ -83,14 +80,14 @@ public:
 	 *
 	 * @return A pointer to the message attachment, or nullptr if no attachment is present.
 	 */
-	virtual IMessageAttachmentPtr GetAttachment( ) const = 0;
+	virtual IMessageAttachmentPtr GetAttachment() const = 0;
 
 	/**
 	 * Gets the date and time at which the message expires.
 	 *
 	 * @return Expiration time.
 	 */
-	virtual const FDateTime& GetExpiration( ) const = 0;
+	virtual const FDateTime& GetExpiration() const = 0;
 
 	/**
 	 * Gets the message address of the endpoint that forwarded this message.
@@ -98,84 +95,84 @@ public:
 	 * @return The forwarder's address.
 	 * @see IsForwarded
 	 */
-	virtual const FMessageAddress& GetForwarder( ) const = 0;
+	virtual const FMessageAddress& GetForwarder() const = 0;
 
 	/**
 	 * Gets the optional message headers.
 	 *
 	 * @return Message header collection.
 	 */
-	virtual const TMap<FName, FString>& GetHeaders( ) const = 0;
+	virtual const TMap<FName, FString>& GetHeaders() const = 0;
 
 	/**
 	 * Gets the message object.
 	 *
 	 * @return The message.
 	 */
-	virtual const void* GetMessage( ) const = 0;
+	virtual const void* GetMessage() const = 0;
 
 	/**
 	 * Gets the name of the message type.
 	 *
 	 * @return Message type name.
 	 */
-	virtual FName GetMessageType( ) const = 0;
+	virtual FName GetMessageType() const = 0;
 
 	/**
 	 * Gets the message's type information.
 	 *
 	 * @return Message type information.
 	 */
-	virtual const TWeakObjectPtr<UScriptStruct>& GetMessageTypeInfo( ) const = 0;
+	virtual const TWeakObjectPtr<UScriptStruct>& GetMessageTypeInfo() const = 0;
 
 	/**
 	 * Returns the original message context in case the message was forwarded.
 	 *
 	 * @return The original message.
 	 */
-	virtual IMessageContextPtr GetOriginalContext( ) const = 0;
+	virtual IMessageContextPtr GetOriginalContext() const = 0;
 
 	/**
 	* Gets the list of message recipients.
 	*
 	* @return Message recipients.
 	 */
-	virtual const TArray<FMessageAddress>& GetRecipients( ) const = 0;
+	virtual const TArray<FMessageAddress>& GetRecipients() const = 0;
 
 	/**
 	 * Gets the message's scope.
 	 *
 	 * @return Scope.
 	 */
-	virtual EMessageScope::Type GetScope( ) const = 0;
+	virtual EMessageScope GetScope() const = 0;
 
 	/**
 	 * Gets the sender's address.
 	 *
 	 * @return Sender address.
 	 */
-	virtual const FMessageAddress& GetSender( ) const = 0;
+	virtual const FMessageAddress& GetSender() const = 0;
 
 	/**
 	 * Gets the name of the thread from which the message was sent.
 	 *
 	 * @return Sender threat name.
 	 */
-	virtual ENamedThreads::Type GetSenderThread( ) const = 0;
+	virtual ENamedThreads::Type GetSenderThread() const = 0;
 
 	/**
 	 * Gets the time at which the message was forwarded.
 	 *
 	 * @return Time forwarded.
 	 */
-	virtual const FDateTime& GetTimeForwarded( ) const = 0;
+	virtual const FDateTime& GetTimeForwarded() const = 0;
 
 	/**
 	 * Gets the time at which the message was sent.
 	 *
 	 * @return Time sent.
 	 */
-	virtual const FDateTime& GetTimeSent( ) const = 0;
+	virtual const FDateTime& GetTimeSent() const = 0;
 
 	/**
 	 * Checks whether this is a forwarded message.
@@ -183,17 +180,17 @@ public:
 	 * @return true if the message was forwarded, false otherwise.
 	 * @see GetForwarder
 	 */
-	virtual bool IsForwarded( ) const = 0;
+	virtual bool IsForwarded() const = 0;
 
 	/**
 	 * Checks whether this context is valid.
 	 *
 	 * @return true if the context is valid, false otherwise.
 	 */
-	virtual bool IsValid( ) const = 0;
+	virtual bool IsValid() const = 0;
 
 public:
 
 	/** Virtual destructor. */
-	virtual ~IMessageContext( ) { }
+	virtual ~IMessageContext() { }
 };
diff --git a/Engine/Source/Runtime/Messaging/Public/Interfaces/IMessageData.h b/Engine/Source/Runtime/Messaging/Public/Interfaces/IMessageData.h
--- a/Engine/Source/Runtime/Messaging/Public/Interfaces/IMessageData.h
+++ b/Engine/Source/Runtime/Messaging/Public/Interfaces/IMessageData.h
@@ -3,23 +3,20 @@
 #pragma once
 
 
-namespace EMessageDataState
+/**
+ * Enumerates message data states.
+ */
+enum class EMessageDataState
 {
-	/**
-	 * Enumerates message data states.
-	 */
-	enum Type
-	{
-		/** The message data is complete. */
-		Complete,
+	/** The message data is complete. */
+	Complete,
 
-		/** The message data is incomplete. */
-		Incomplete,
+	/** The message data is incomplete. */
+	Incomplete,
 
-		/** The message data is invalid. */
-		Invalid
-	};
-}
+	/** The message data is invalid. */
+	Invalid
+};
 
 
 /** Type definition for shared pointers to instances of IMessageData. */
@@ -43,14 +40,14 @@ public:
 	 *
 	 * @return An archive reader.
 	 */
-	virtual FArchive* CreateReader( ) = 0;
+	virtual FArchive* CreateReader() = 0;
 
 	/**
 	 * Gets the state of the message data.
 	 *
 	 * @return Message data state.
 	 */
-	virtual EMessageDataState::Type GetState( ) const = 0;
+	virtual EMessageDataState GetState() const = 0;
 
 public:
 
@@ -59,10 +56,10 @@ public:
 	 *
 	 * @return The delegate.
 	 */
-	virtual FSimpleDelegate& OnStateChanged( ) = 0;
+	virtual FSimpleDelegate& OnStateChanged() = 0;
 
 public:
 
 	/** Virtual destructor. */
-	virtual ~IMessageData( ) { }
+	virtual ~IMessageData() { }
 };
diff --git a/Engine/Source/Runtime/Messaging/Public/Interfaces/IMessageTracer.h b/Engine/Source/Runtime/Messaging/Public/Interfaces/IMessageTracer.h
--- a/Engine/Source/Runtime/Messaging/Public/Interfaces/IMessageTracer.h
+++ b/Engine/Source/Runtime/Messaging/Public/Interfaces/IMessageTracer.h
@@ -40,23 +40,20 @@ typedef TSharedPtr<struct FMessageTracerTypeInfo> FMessageTracerTypeInfoPtr;
 typedef TSharedRef<struct FMessageTracerTypeInfo> FMessageTracerTypeInfoRef;
 
 
-namespace EMessageTracerDispatchTypes
+/**
+ * Enumerates message dispatch types.
+ */
+enum class EMessageTracerDispatchTypes
 {
-	/**
-	 * Enumerates message dispatch types.
-	 */
-	enum Type
-	{
-		/** The message is being dispatched directly. */
-		Direct,
+	/** The message is being dispatched directly. */
+	Direct,
 
-		/** The message hasn't been dispatched yet. */
-		Pending,
+	/** The message hasn't been dispatched yet. */
+	Pending,
 
-		/** The message is being dispatched using the task graph system. */
-		TaskGraph
-	};
-}
+	/** The message is being dispatched using the task graph system. */
+	TaskGraph
+};
 
 
 /**
@@ -68,7 +65,7 @@ struct FMessageTracerDispatchState
 	double DispatchLatency;
 
 	/** Holds the message's dispatch type for the specified endpoint. */
-	EMessageTracerDispatchTypes::Type DispatchType;
+	EMessageTracerDispatchTypes DispatchType;
 
 	/** Holds the endpoint to which the message was or is being dispatched. */
 	FMessageTracerEndpointInfoPtr EndpointInfo;
@@ -188,14 +185,14 @@ public:
 	 *
 	 * @see Continue, IsBreaking, Step
 	 */
-	virtual void Break( ) = 0;
+	virtual void Break() = 0;
 
 	/**
 	 * Continues message routing from the current breakpoint.
 	 *
 	 * @see Break, Step
 	 */
-	virtual void Continue( ) = 0;
+	virtual void Continue() = 0;
 
 	/**
 	 * Checks whether the tracer is currently at a breakpoint.
@@ -203,7 +200,7 @@ public:
 	 * @return true if at breakpoint, false otherwise.
 	 * @see Break
 	 */
-	virtual bool IsBreaking( ) const = 0;
+	virtual bool IsBreaking() const = 0;
 
 	/**
 	 * Checks whether the tracer is currently running.
@@ -212,33 +209,33 @@ public:
 	 *
 	 * @see Start, Stop
 	 */
-	virtual bool IsRunning( ) const = 0;
+	virtual bool IsRunning() const = 0;
 
 	/**
 	 * Resets the tracer.
 	 */
-	virtual void Reset( ) = 0;
+	virtual void Reset() = 0;
 
 	/**
 	 * Starts the tracer.
 	 *
 	 * @see IsRunning, Stop
 	 */
-	virtual void Start( ) = 0;
+	virtual void Start() = 0;
 
 	/**
 	 * Steps the tracer to the next message.
 	 *
 	 * @see Break, Continue
 	 */
-	virtual void Step( ) = 0;
+	virtual void Step() = 0;
 
 	/**
 	 * Stops the tracer.
 	 *
 	 * @see IsRunning, Start
 	 */
-	virtual void Stop( ) = 0;
+	virtual void Stop() = 0;
 
 	/**
 	 * Ticks the tracer.
@@ -279,7 +276,7 @@ public:
 	 *
 	 * @return true if there are messages, false otherwise.
 	 */
-	virtual bool HasMessages( ) const = 0;
+	virtual bool HasMessages() const = 0;
 
 public:
 
@@ -288,24 +285,24 @@ public:
 	 *
 	 * @return The delegate.
 	 */
-	virtual FMessageTracerMessageAdded& OnMessageAdded( ) = 0;
+	virtual FMessageTracerMessageAdded& OnMessageAdded() = 0;
 
 	/**
 	 * Returns a delegate that is executed when the message history has been reset.
 	 *
 	 * @return The delegate.
 	 */
-	virtual FSimpleMulticastDelegate& OnMessagesReset( ) = 0;
+	virtual FSimpleMulticastDelegate& OnMessagesReset() = 0;
 
 	/**
 	 * Returns a delegate that is executed when the collection of known messages types has changed.
 	 *
 	 * @return The delegate.
 	 */
-	virtual FMessageTracerTypeAdded& OnTypeAdded( ) = 0;
+	virtual FMessageTracerTypeAdded& OnTypeAdded() = 0;
 
 protected:
 
 	/** Hidden destructor. The life time of a message tracer is managed by the message bus. */
-	virtual ~IMessageTracer( ) { }
+	virtual ~IMessageTracer() { }
 };
diff --git a/Engine/Source/Runtime/Messaging/Public/Interfaces/IMutableMessageContext.h b/Engine/Source/Runtime/Messaging/Public/Interfaces/IMutableMessageContext.h
--- a/Engine/Source/Runtime/Messaging/Public/Interfaces/IMutableMessageContext.h
+++ b/Engine/Source/Runtime/Messaging/Public/Interfaces/IMutableMessageContext.h
@@ -60,7 +60,7 @@ public:
 	 *
 	 * @param InScope The message scope.
 	 */
-	virtual void SetScope( EMessageScope::Type InScope ) = 0;
+	virtual void SetScope( EMessageScope InScope ) = 0;
 
 	/**
 	 * Sets the address of the message's sender.
@@ -79,5 +79,5 @@ public:
 public:
 
 	/** Virtual destructor. */
-	virtual ~IMutableMessageContext( ) { }
+	virtual ~IMutableMessageContext() { }
 };
