Messaging: Fixed message context not deleting message memory

diff --git a/Engine/Source/Runtime/Messaging/Private/Bus/MessageContext.h b/Engine/Source/Runtime/Messaging/Private/Bus/MessageContext.h
--- a/Engine/Source/Runtime/Messaging/Private/Bus/MessageContext.h
+++ b/Engine/Source/Runtime/Messaging/Private/Bus/MessageContext.h
@@ -71,9 +71,14 @@ public:
 	/** Destructor. */
 	~FMessageContext()
 	{
-		if ((Message != nullptr) && TypeInfo.IsValid())
+		if (Message != nullptr)
 		{
-			TypeInfo->DestroyScriptStruct(Message);
+			if (TypeInfo.IsValid())
+			{
+				TypeInfo->DestroyScriptStruct(Message);
+			}
+
+			delete Message;
 		}		
 	}
 
