Messaging: Removed unused dependencies to Sockets module

diff --git a/Engine/Source/Runtime/Messaging/Messaging.Build.cs b/Engine/Source/Runtime/Messaging/Messaging.Build.cs
--- a/Engine/Source/Runtime/Messaging/Messaging.Build.cs
+++ b/Engine/Source/Runtime/Messaging/Messaging.Build.cs
@@ -18,7 +18,6 @@ namespace UnrealBuildTool.Rules
 				{
 					"Core",
 					"Json",
-					"Sockets",
 				}
 			);
 
diff --git a/Engine/Source/Runtime/Messaging/Private/MessagingPrivatePCH.h b/Engine/Source/Runtime/Messaging/Private/MessagingPrivatePCH.h
--- a/Engine/Source/Runtime/Messaging/Private/MessagingPrivatePCH.h
+++ b/Engine/Source/Runtime/Messaging/Private/MessagingPrivatePCH.h
@@ -9,8 +9,6 @@
  *****************************************************************************/
 
 #include "Json.h"
-#include "Sockets.h"
-#include "SocketSubsystem.h"
 
 
 /* Private includes
diff --git a/Engine/Source/Runtime/Messaging/Public/Interfaces/IMessageData.h b/Engine/Source/Runtime/Messaging/Public/Interfaces/IMessageData.h
--- a/Engine/Source/Runtime/Messaging/Public/Interfaces/IMessageData.h
+++ b/Engine/Source/Runtime/Messaging/Public/Interfaces/IMessageData.h
@@ -46,6 +46,7 @@ public:
 	 * Gets the state of the message data.
 	 *
 	 * @return Message data state.
+	 * @see OnStateChanged
 	 */
 	virtual EMessageDataState GetState() const = 0;
 
@@ -55,6 +56,7 @@ public:
 	 * Returns a delegate that is executed when the message data's state changed.
 	 *
 	 * @return The delegate.
+	 * @see GetState
 	 */
 	virtual FSimpleDelegate& OnStateChanged() = 0;
 
diff --git a/Engine/Source/Runtime/Messaging/Public/Interfaces/IMessageTracer.h b/Engine/Source/Runtime/Messaging/Public/Interfaces/IMessageTracer.h
--- a/Engine/Source/Runtime/Messaging/Public/Interfaces/IMessageTracer.h
+++ b/Engine/Source/Runtime/Messaging/Public/Interfaces/IMessageTracer.h
@@ -211,9 +211,7 @@ public:
 	 */
 	virtual bool IsRunning() const = 0;
 
-	/**
-	 * Resets the tracer.
-	 */
+	/** Resets the tracer. */
 	virtual void Reset() = 0;
 
 	/**
