Messaging: First pass on making headers compilable individually

diff --git a/Engine/Plugins/2D/Paper2D/Source/Paper2DEditor/Private/TileSetAssetTypeActions.cpp b/Engine/Plugins/2D/Paper2D/Source/Paper2DEditor/Private/TileSetAssetTypeActions.cpp
--- a/Engine/Plugins/2D/Paper2D/Source/Paper2DEditor/Private/TileSetAssetTypeActions.cpp
+++ b/Engine/Plugins/2D/Paper2D/Source/Paper2DEditor/Private/TileSetAssetTypeActions.cpp
@@ -1,16 +1,16 @@
 // Copyright 1998-2014 Epic Games, Inc. All Rights Reserved.
 
 #include "Paper2DEditorPrivatePCH.h"
-
-//#include "Messaging.h"
 //#include "AssetToolsModule.h"
 //#include "Toolkits/AssetEditorManager.h"
 //#include "AssetRegistryModule.h"
 #include "TileSetAssetTypeActions.h"
 #include "TileSetEditor.h"
 
+
 #define LOCTEXT_NAMESPACE "AssetTypeActions"
 
+
 //////////////////////////////////////////////////////////////////////////
 // FTileSetAssetTypeActions
 
diff --git a/Engine/Source/Developer/AutomationController/Private/AutomationDeviceClusterManager.cpp b/Engine/Source/Developer/AutomationController/Private/AutomationDeviceClusterManager.cpp
--- a/Engine/Source/Developer/AutomationController/Private/AutomationDeviceClusterManager.cpp
+++ b/Engine/Source/Developer/AutomationController/Private/AutomationDeviceClusterManager.cpp
@@ -233,7 +233,7 @@ bool FAutomationDeviceClusterManager::FindDevice(const FMessageAddress& MessageA
 }
 
 
-FGuid FAutomationDeviceClusterManager::GetDeviceMessageAddress(const int32 ClusterIndex, const int32 DeviceIndex) const
+FMessageAddress FAutomationDeviceClusterManager::GetDeviceMessageAddress(const int32 ClusterIndex, const int32 DeviceIndex) const
 {
 	//verify cluster/device index
 	check((ClusterIndex >= 0) && (ClusterIndex < Clusters.Num()));
diff --git a/Engine/Source/Developer/TargetDeviceServices/Private/TargetDeviceService.cpp b/Engine/Source/Developer/TargetDeviceServices/Private/TargetDeviceService.cpp
--- a/Engine/Source/Developer/TargetDeviceServices/Private/TargetDeviceService.cpp
+++ b/Engine/Source/Developer/TargetDeviceServices/Private/TargetDeviceService.cpp
@@ -295,7 +295,7 @@ void FTargetDeviceService::HandleUnclaimedMessage(const FTargetDeviceUnclaimed&
 	{
 		if (Context->GetSender() == ClaimAddress)
 		{
-			ClaimAddress = FGuid();
+			ClaimAddress.Invalidate();
 			ClaimHost.Empty();
 			ClaimUser.Empty();
 		}
diff --git a/Engine/Source/Runtime/AutomationWorker/Private/AutomationWorkerModule.cpp b/Engine/Source/Runtime/AutomationWorker/Private/AutomationWorkerModule.cpp
--- a/Engine/Source/Runtime/AutomationWorker/Private/AutomationWorkerModule.cpp
+++ b/Engine/Source/Runtime/AutomationWorker/Private/AutomationWorkerModule.cpp
@@ -138,7 +138,7 @@ void FAutomationWorkerModule::ReportNetworkCommandComplete()
 {
 	if (GIsAutomationTesting)
 	{
-		MessageEndpoint->Send(new FAutomationWorkerRequestNextNetworkCommand(ExecutionCount), TestRequesterGUID);
+		MessageEndpoint->Send(new FAutomationWorkerRequestNextNetworkCommand(ExecutionCount), TestRequesterAddress);
 		if (StopTestEvent.IsBound())
 		{
 			// this is a local test; the message to continue will never arrive, so lets not wait for it
@@ -212,12 +212,12 @@ void FAutomationWorkerModule::ReportTestComplete()
 			Message->Warnings = ExecutionInfo.Warnings;
 			Message->Logs = ExecutionInfo.LogItems;
 
-			MessageEndpoint->Send(Message, TestRequesterGUID);
+			MessageEndpoint->Send(Message, TestRequesterAddress);
 		}
 
 
 		// reset local state
-		TestRequesterGUID.Invalidate();
+		TestRequesterAddress.Invalidate();
 		ExecutionCount = INDEX_NONE;
 		TestName.Empty();
 		StopTestEvent.Unbind();
@@ -339,14 +339,14 @@ void FAutomationWorkerModule::HandleScreenShotCaptured( int32 Width, int32 Heigh
 		FImageUtils::CompressImageArray(NewWidth, NewHeight, ScaledBitmap, CompressedBitmap);
 
 		// Send the screen shot if we have a target
-		if( TestRequesterGUID.IsValid() )
+		if( TestRequesterAddress.IsValid() )
 		{
 			FAutomationWorkerScreenImage* Message = new FAutomationWorkerScreenImage();
 
 			FString SFilename = ScreenShotName;
 			Message->ScreenShotName = SFilename;
 			Message->ScreenImage = CompressedBitmap;
-			MessageEndpoint->Send(Message, TestRequesterGUID);
+			MessageEndpoint->Send(Message, TestRequesterAddress);
 		}
 		else
 		{
@@ -366,7 +366,7 @@ void FAutomationWorkerModule::HandleRunTestsMessage( const FAutomationWorkerRunT
 {
 	ExecutionCount = Message.ExecutionCount;
 	TestName = Message.TestName;
-	TestRequesterGUID = Context->GetSender();
+	TestRequesterAddress = Context->GetSender();
 	FAutomationTestFramework::GetInstance().SetScreenshotOptions(Message.bScreenshotsEnabled, Message.bUseFullSizeScreenShots);
 
 	// Always allow the first network command to execute
diff --git a/Engine/Source/Runtime/AutomationWorker/Private/AutomationWorkerModule.h b/Engine/Source/Runtime/AutomationWorker/Private/AutomationWorkerModule.h
--- a/Engine/Source/Runtime/AutomationWorker/Private/AutomationWorkerModule.h
+++ b/Engine/Source/Runtime/AutomationWorker/Private/AutomationWorkerModule.h
@@ -98,16 +98,16 @@ private:
 
 private:
 
-	// Holds the messaging endpoint.
+	/** Holds the messaging endpoint. */
 	FMessageEndpointPtr MessageEndpoint;
 
-	//ID of the controller sending the test request
-	FGuid TestRequesterGUID;
+	/** Message address of the controller sending the test request. */
+	FMessageAddress TestRequesterAddress;
 
-	//identifier for the controller to know if the results should be discarded or not
+	/** Identifier for the controller to know if the results should be discarded or not. */
 	uint32 ExecutionCount;
 
-	//execute one of the tests by request of the controller
+	/** Execute one of the tests by request of the controller. */
 	FString TestName;
 
 	/** Whether the controller has requested that the network command should execute */
diff --git a/Engine/Source/Runtime/Core/Public/Misc/Guid.h b/Engine/Source/Runtime/Core/Public/Misc/Guid.h
--- a/Engine/Source/Runtime/Core/Public/Misc/Guid.h
+++ b/Engine/Source/Runtime/Core/Public/Misc/Guid.h
@@ -207,26 +207,26 @@ public:
 	CORE_API bool ImportTextItem( const TCHAR*& Buffer, int32 PortFlags, class UObject* Parent, FOutputDevice* ErrorText );
 
 	/**
-	 * Checks whether this GUID is valid or not.
-	 *
-	 * A GUID that has all its components set to zero is considered invalid.
+	 * Invalidates the GUID.
 	 *
-	 * @return true if valid, false otherwise
-	 * @see Invalidate
+	 * @see IsValid
 	 */
-	bool IsValid() const
+	void Invalidate()
 	{
-		return ((A | B | C | D) != 0);
+		A = B = C = D = 0;
 	}
 
 	/**
-	 * Invalidates the GUID.
+	 * Checks whether this GUID is valid or not.
 	 *
-	 * @see IsValid
+	 * A GUID that has all its components set to zero is considered invalid.
+	 *
+	 * @return true if valid, false otherwise.
+	 * @see Invalidate
 	 */
-	void Invalidate()
+	bool IsValid() const
 	{
-		A = B = C = D = 0;
+		return ((A | B | C | D) != 0);
 	}
 
 	/**
@@ -257,7 +257,7 @@ public:
 	 */
 	friend uint32 GetTypeHash( const FGuid& Guid )
 	{
-		return FCrc::MemCrc_DEPRECATED(&Guid,sizeof(FGuid));
+		return FCrc::MemCrc_DEPRECATED(&Guid, sizeof(FGuid));
 	}
 
 public:
diff --git a/Engine/Source/Runtime/Engine/Private/EngineService.cpp b/Engine/Source/Runtime/Engine/Private/EngineService.cpp
--- a/Engine/Source/Runtime/Engine/Private/EngineService.cpp
+++ b/Engine/Source/Runtime/Engine/Private/EngineService.cpp
@@ -1,13 +1,10 @@
 // Copyright 1998-2014 Epic Games, Inc. All Rights Reserved.
 
-/*=============================================================================
-	EngineService.cpp: Implements the FEngineService class.
-=============================================================================*/
-
 #include "EnginePrivate.h"
+#include "EngineServiceMessages.h"
 #include "MessageEndpoint.h"
 #include "MessageEndpointBuilder.h"
-#include "EngineServiceMessages.h"
+
 
 DEFINE_LOG_CATEGORY_STATIC(LogEngineService, Log, All)
 
@@ -15,7 +12,7 @@ DEFINE_LOG_CATEGORY_STATIC(LogEngineService, Log, All)
 /* FEngineService structors
  *****************************************************************************/
 
-FEngineService::FEngineService( )
+FEngineService::FEngineService()
 {
 	// always grant access to session owner
 	AuthorizedUsers.Add(FApp::GetSessionOwner());
@@ -58,7 +55,7 @@ void FEngineService::SendPong( const IMessageContextRef& Context )
 		Message->InstanceId = FApp::GetInstanceId();
 		Message->SessionId = FApp::GetSessionId();
 
-		if (GEngine == NULL)
+		if (GEngine == nullptr)
 		{
 			Message->InstanceType = TEXT("Unknown");
 		}
@@ -83,7 +80,7 @@ void FEngineService::SendPong( const IMessageContextRef& Context )
 			Message->InstanceType = TEXT("Other");
 		}
 
-		FWorldContext const* ContextToUse = NULL;
+		FWorldContext const* ContextToUse = nullptr;
 
 		// TODO: Should we be iterating here and sending a message for each context?
 
@@ -97,17 +94,17 @@ void FEngineService::SendPong( const IMessageContextRef& Context )
 				ContextToUse = &WorldContext;
 				break;
 			}
-			else if (WorldContext.WorldType == EWorldType::PIE && (ContextToUse == NULL || ContextToUse->WorldType != EWorldType::PIE))
+			else if (WorldContext.WorldType == EWorldType::PIE && (ContextToUse == nullptr || ContextToUse->WorldType != EWorldType::PIE))
 			{
 				ContextToUse = &WorldContext;
 			}
-			else if (WorldContext.WorldType == EWorldType::Editor && ContextToUse == NULL)
+			else if (WorldContext.WorldType == EWorldType::Editor && ContextToUse == nullptr)
 			{	
 				ContextToUse = &WorldContext;
 			}
 		}
 
-		if (ContextToUse != NULL && ContextToUse->World() != NULL)
+		if (ContextToUse != nullptr && ContextToUse->World() != nullptr)
 		{
 			Message->CurrentLevel = ContextToUse->World()->GetMapName();
 			Message->HasBegunPlay = ContextToUse->World()->HasBegunPlay();
@@ -156,7 +153,7 @@ void FEngineService::HandleExecuteCommandMessage( const FEngineServiceExecuteCom
 {
 	if (AuthorizedUsers.Contains(Message.UserName))
 	{
-		if (GEngine != NULL)
+		if (GEngine != nullptr)
 		{
 			GEngine->DeferredCommands.Add(Message.Command);
 
@@ -184,17 +181,17 @@ void FEngineService::HandleTerminateMessage( const FEngineServiceTerminate& Mess
 {
 	if (AuthorizedUsers.Contains(Message.UserName))
 	{
-		if (GEngine != NULL)
+		if (GEngine != nullptr)
 		{
 			UE_LOG(LogEngineService, Log, TEXT("%s terminated this instance remotely."), *Message.UserName);
 
 			if (GEngine->IsEditor())
 			{
-				GEngine->Exec( NULL, TEXT("QUIT_EDITOR"), *GLog);
+				GEngine->Exec( nullptr, TEXT("QUIT_EDITOR"), *GLog);
 			}
 			else
 			{
-				GEngine->Exec( NULL, TEXT("QUIT"), *GLog);
+				GEngine->Exec( nullptr, TEXT("QUIT"), *GLog);
 			}
 		}
 		else
diff --git a/Engine/Source/Runtime/Engine/Public/EngineService.h b/Engine/Source/Runtime/Engine/Public/EngineService.h
--- a/Engine/Source/Runtime/Engine/Public/EngineService.h
+++ b/Engine/Source/Runtime/Engine/Public/EngineService.h
@@ -1,13 +1,17 @@
 // Copyright 1998-2014 Epic Games, Inc. All Rights Reserved.
 
-/*=============================================================================
-	EngineService.h: Declares the FEngineService class.
-=============================================================================*/
-
 #pragma once
 
-#include "Messaging.h"
-#include "EngineMessages.h"
+
+struct FEngineServiceAuthDeny;
+struct FEngineServiceAuthGrant;
+struct FEngineServiceExecuteCommand;
+struct FEngineServicePing;
+struct FEngineServiceTerminate;
+struct FMessageAddress;
+class FMessageEndpoint;
+class IMessageContext;
+
 
 /**
  * Implements an application session service.
@@ -16,50 +20,48 @@ class FEngineService
 {
 public:
 
-	/**
-	 * Default constructor.
-	 */
-	ENGINE_API FEngineService( );
+	/** Default constructor. */
+	ENGINE_API FEngineService();
 
 protected:
 
 	/**
 	 * Sends a notification to the specified recipient.
 	 *
-	 * @param NotificationText - The notification text.
-	 * @param Recipient - The recipient's message address.
+	 * @param NotificationText The notification text.
+	 * @param Recipient The recipient's message address.
 	 */
 	void SendNotification( const TCHAR* NotificationText, const FMessageAddress& Recipient );
 
 	/**
 	 * Publishes a ping response.
 	 *
-	 * @param Context - The message context of the received Ping message.
+	 * @param Context The message context of the received Ping message.
 	 */
-	void SendPong( const IMessageContextRef& Context );
+	void SendPong( const TSharedRef<IMessageContext, ESPMode::ThreadSafe>& Context );
 
 private:
 
-	// Handles FEngineServiceAuthGrant messages.
-	void HandleAuthGrantMessage( const FEngineServiceAuthGrant& Message, const IMessageContextRef& Context );
+	/** Handles FEngineServiceAuthGrant messages. */
+	void HandleAuthGrantMessage( const FEngineServiceAuthGrant& Message, const TSharedRef<IMessageContext, ESPMode::ThreadSafe>& Context );
 
-	// Handles FEngineServiceAuthGrant messages.
-	void HandleAuthDenyMessage( const FEngineServiceAuthDeny& Message, const IMessageContextRef& Context );
+	/** Handles FEngineServiceAuthGrant messages. */
+	void HandleAuthDenyMessage( const FEngineServiceAuthDeny& Message, const TSharedRef<IMessageContext, ESPMode::ThreadSafe>& Context );
 
-	// Handles FEngineServiceExecuteCommand messages.
-	void HandleExecuteCommandMessage( const FEngineServiceExecuteCommand& Message, const IMessageContextRef& Context );
+	/** Handles FEngineServiceExecuteCommand messages. */
+	void HandleExecuteCommandMessage( const FEngineServiceExecuteCommand& Message, const TSharedRef<IMessageContext, ESPMode::ThreadSafe>& Context );
 
-	// Handles FEngineServicePing messages.
-	void HandlePingMessage( const FEngineServicePing& Message, const IMessageContextRef& Context );
+	/** Handles FEngineServicePing messages. */
+	void HandlePingMessage( const FEngineServicePing& Message, const TSharedRef<IMessageContext, ESPMode::ThreadSafe>& Context );
 
-	// Handles FEngineServiceTerminate messages.
-	void HandleTerminateMessage( const FEngineServiceTerminate& Message, const IMessageContextRef& Context );
+	/** Handles FEngineServiceTerminate messages. */
+	void HandleTerminateMessage( const FEngineServiceTerminate& Message, const TSharedRef<IMessageContext, ESPMode::ThreadSafe>& Context );
 
 private:
 
-	// Holds the list of users that are allowed to interact with this session.
+	/** Holds the list of users that are allowed to interact with this session. */
 	TArray<FString> AuthorizedUsers;
 
-	// Holds the message endpoint.
-	FMessageEndpointPtr MessageEndpoint;
+	/** Holds the message endpoint. */
+	TSharedPtr<FMessageEndpoint, ESPMode::ThreadSafe> MessageEndpoint;
 };
diff --git a/Engine/Source/Runtime/Messaging/Public/Common/MessageBridgeBuilder.h b/Engine/Source/Runtime/Messaging/Public/Common/MessageBridgeBuilder.h
--- a/Engine/Source/Runtime/Messaging/Public/Common/MessageBridgeBuilder.h
+++ b/Engine/Source/Runtime/Messaging/Public/Common/MessageBridgeBuilder.h
@@ -12,7 +12,7 @@ public:
 
 	/** Default constructor. */
 	FMessageBridgeBuilder()
-		: Address(FMessageAddress::NewGuid())
+		: Address(FMessageAddress::NewAddress())
 		, BusPtr(IMessagingModule::Get().GetDefaultBus())
 		, Disabled(false)
 		, Transport(nullptr)
@@ -24,7 +24,7 @@ public:
 	 * @param InBus The message bus to attach the bridge to.
 	 */
 	FMessageBridgeBuilder( const IMessageBusRef& Bus )
-		: Address(FMessageAddress::NewGuid())
+		: Address(FMessageAddress::NewAddress())
 		, BusPtr(Bus)
 		, Disabled(false)
 		, Transport(nullptr)
diff --git a/Engine/Source/Runtime/Messaging/Public/Common/MessageEndpoint.h b/Engine/Source/Runtime/Messaging/Public/Common/MessageEndpoint.h
--- a/Engine/Source/Runtime/Messaging/Public/Common/MessageEndpoint.h
+++ b/Engine/Source/Runtime/Messaging/Public/Common/MessageEndpoint.h
@@ -2,12 +2,12 @@
 
 #pragma once
 
-
-/** Type definition for shared pointers to instances of FMessageEndpoint. */
-typedef TSharedPtr<class FMessageEndpoint, ESPMode::ThreadSafe> FMessageEndpointPtr;
-
-/** Type definition for shared references to instances of FMessageEndpoint. */
-typedef TSharedRef<class FMessageEndpoint, ESPMode::ThreadSafe> FMessageEndpointRef;
+#include "IMessageAttachment.h"
+#include "IMessageBus.h"
+#include "IMessageContext.h"
+#include "IMessageHandler.h"
+#include "IReceiveMessages.h"
+#include "ISendMessages.h"
 
 
 /**
@@ -67,7 +67,7 @@ public:
 	 * @param InHandlers The collection of message handlers to register.
 	 */
 	FMessageEndpoint( const FName& InName, const IMessageBusRef& InBus, const TArray<IMessageHandlerPtr>& InHandlers )
-		: Address(FGuid::NewGuid())
+		: Address(FMessageAddress::NewAddress())
 		, BusPtr(InBus)
 		, Enabled(true)
 		, Handlers(InHandlers)
@@ -95,7 +95,7 @@ public:
 	 * A disabled endpoint will not receive any subscribed messages until it is enabled again.
 	 * Endpoints should be created in an enabled state by default.
 	 *
-	 * @see Enable
+	 * @see Enable, IsEnabled
 	 */
 	void Disable()
 	{
@@ -108,7 +108,7 @@ public:
 	 * An activated endpoint will receive subscribed messages.
 	 * Endpoints should be created in an enabled state by default.
 	 *
-	 * @see Disable
+	 * @see Disable, IsEnabled
 	 */
 	void Enable()
 	{
@@ -127,6 +127,9 @@ public:
 
 	/**
 	 * Checks whether this endpoint is connected to the bus.
+	 *
+	 * @return true if connected, false otherwise.
+	 * @see IsEnabled
 	 */
 	bool IsConnected() const
 	{
@@ -137,6 +140,7 @@ public:
 	 * Checks whether this endpoint is enabled.
 	 *
 	 * @return true if the endpoint is enabled, false otherwise.
+	 * @see Disable, Enable, IsConnected
 	 */
 	bool IsEnabled() const
 	{
@@ -741,7 +745,7 @@ public:
 	 *
 	 * @param Endpoint The message endpoint to release.
 	 */
-	static void SafeRelease( FMessageEndpointPtr& Endpoint )
+	static void SafeRelease( TSharedPtr<FMessageEndpoint, ESPMode::ThreadSafe>& Endpoint )
 	{
 		TWeakPtr<FMessageEndpoint, ESPMode::ThreadSafe> EndpointPtr = Endpoint;
 		Endpoint.Reset();
@@ -825,3 +829,10 @@ private:
 	/** Holds a delegate that is invoked in case of messaging errors. */
 	FOnMessageEndpointError ErrorDelegate;
 };
+
+
+/** Type definition for shared pointers to instances of FMessageEndpoint. */
+typedef TSharedPtr<FMessageEndpoint, ESPMode::ThreadSafe> FMessageEndpointPtr;
+
+/** Type definition for shared references to instances of FMessageEndpoint. */
+typedef TSharedRef<FMessageEndpoint, ESPMode::ThreadSafe> FMessageEndpointRef;
diff --git a/Engine/Source/Runtime/Messaging/Public/Common/MessageEndpointBuilder.h b/Engine/Source/Runtime/Messaging/Public/Common/MessageEndpointBuilder.h
--- a/Engine/Source/Runtime/Messaging/Public/Common/MessageEndpointBuilder.h
+++ b/Engine/Source/Runtime/Messaging/Public/Common/MessageEndpointBuilder.h
@@ -2,6 +2,9 @@
 
 #pragma once
 
+#include "IMessagingModule.h"
+#include "MessageEndpoint.h"
+
 
 /**
  * Implements a message endpoint builder.
diff --git a/Engine/Source/Runtime/Messaging/Public/Interfaces/IAuthorizeMessageRecipients.h b/Engine/Source/Runtime/Messaging/Public/Interfaces/IAuthorizeMessageRecipients.h
--- a/Engine/Source/Runtime/Messaging/Public/Interfaces/IAuthorizeMessageRecipients.h
+++ b/Engine/Source/Runtime/Messaging/Public/Interfaces/IAuthorizeMessageRecipients.h
@@ -3,11 +3,9 @@
 #pragma once
 
 
-/** Type definition for shared pointers to instances of IAuthorizeMessageRecipients. */
-typedef TSharedPtr<class IAuthorizeMessageRecipients> IAuthorizeMessageRecipientsPtr;
-
-/** Type definition for shared references to instances of IAuthorizeMessageRecipients. */
-typedef TSharedRef<class IAuthorizeMessageRecipients> IAuthorizeMessageRecipientsRef;
+struct FMessageAddress;
+class IInterceptMessages;
+class IReceiveMessages;
 
 
 /**
@@ -24,7 +22,7 @@ public:
 	 * @param MessageType The type of messages to intercept.
 	 * @return true if the request was authorized, false otherwise.
 	 */
-	virtual bool AuthorizeInterceptor( const IInterceptMessagesRef& Interceptor, const FName& MessageType ) = 0;
+	virtual bool AuthorizeInterceptor( const TSharedRef<IInterceptMessages, ESPMode::ThreadSafe>& Interceptor, const FName& MessageType ) = 0;
 
 	/**
 	 * Authorizes a request to register the specified recipient.
@@ -33,7 +31,7 @@ public:
 	 * @param Address The recipient's address.
 	 * @return true if the request was authorized, false otherwise.
 	 */
-	virtual bool AuthorizeRegistration( const IReceiveMessagesRef& Recipient, const FMessageAddress& Address ) = 0;
+	virtual bool AuthorizeRegistration( const TSharedRef<IReceiveMessages, ESPMode::ThreadSafe>& Recipient, const FMessageAddress& Address ) = 0;
 
 	/**
 	 * Authorizes a request to add a subscription for the specified topic pattern.
@@ -42,7 +40,7 @@ public:
 	 * @param TopicPattern The message topic pattern to subscribe to.
 	 * @return true if the request is authorized, false otherwise.
 	 */
-	virtual bool AuthorizeSubscription( const IReceiveMessagesRef& Subscriber, const FName& TopicPattern ) = 0;
+	virtual bool AuthorizeSubscription( const TSharedRef<IReceiveMessages, ESPMode::ThreadSafe>& Subscriber, const FName& TopicPattern ) = 0;
 
 	/**
 	 * Authorizes a request to unregister the specified recipient.
@@ -59,10 +57,17 @@ public:
 	 * @param TopicPattern The message topic pattern to unsubscribe from.
 	 * @return true if the request is authorized, false otherwise.
 	 */
-	virtual bool AuthorizeUnsubscription( const IReceiveMessagesRef& Subscriber, const FName& TopicPattern ) = 0;
+	virtual bool AuthorizeUnsubscription( const TSharedRef<IReceiveMessages, ESPMode::ThreadSafe>& Subscriber, const FName& TopicPattern ) = 0;
 
 public:
 
 	/** Virtual destructor. */
 	virtual ~IAuthorizeMessageRecipients() { }
 };
+
+
+/** Type definition for shared pointers to instances of IAuthorizeMessageRecipients. */
+typedef TSharedPtr<IAuthorizeMessageRecipients> IAuthorizeMessageRecipientsPtr;
+
+/** Type definition for shared references to instances of IAuthorizeMessageRecipients. */
+typedef TSharedRef<IAuthorizeMessageRecipients> IAuthorizeMessageRecipientsRef;
diff --git a/Engine/Source/Runtime/Messaging/Public/Interfaces/IInterceptMessages.h b/Engine/Source/Runtime/Messaging/Public/Interfaces/IInterceptMessages.h
--- a/Engine/Source/Runtime/Messaging/Public/Interfaces/IInterceptMessages.h
+++ b/Engine/Source/Runtime/Messaging/Public/Interfaces/IInterceptMessages.h
@@ -3,11 +3,7 @@
 #pragma once
 
 
-/** Type definition for shared pointers to instances of IInterceptMessages. */
-typedef TSharedPtr<class IInterceptMessages, ESPMode::ThreadSafe> IInterceptMessagesPtr;
-
-/** Type definition for shared references to instances of IInterceptMessages. */
-typedef TSharedRef<class IInterceptMessages, ESPMode::ThreadSafe> IInterceptMessagesRef;
+class IMessageContext;
 
 
 /**
@@ -23,10 +19,17 @@ public:
 	 * @param Context The context of the message to intercept.
 	 * @return true if the message was intercepted and should not be routed, false otherwise.
 	 */
-	virtual bool InterceptMessage( const IMessageContextRef& Context ) = 0;
+	virtual bool InterceptMessage( const TSharedRef<IMessageContext, ESPMode::ThreadSafe>& Context ) = 0;
 
 public:
 
 	/** Virtual destructor. */
 	virtual ~IInterceptMessages() { }
 };
+
+
+/** Type definition for shared pointers to instances of IInterceptMessages. */
+typedef TSharedPtr<IInterceptMessages, ESPMode::ThreadSafe> IInterceptMessagesPtr;
+
+/** Type definition for shared references to instances of IInterceptMessages. */
+typedef TSharedRef<IInterceptMessages, ESPMode::ThreadSafe> IInterceptMessagesRef;
diff --git a/Engine/Source/Runtime/Messaging/Public/Interfaces/IMessageAttachment.h b/Engine/Source/Runtime/Messaging/Public/Interfaces/IMessageAttachment.h
--- a/Engine/Source/Runtime/Messaging/Public/Interfaces/IMessageAttachment.h
+++ b/Engine/Source/Runtime/Messaging/Public/Interfaces/IMessageAttachment.h
@@ -3,13 +3,6 @@
 #pragma once
 
 
-/** Type definition for shared pointers to instances of IMessageAttachment. */
-typedef TSharedPtr<class IMessageAttachment, ESPMode::ThreadSafe> IMessageAttachmentPtr;
-
-/** Type definition for shared references to instances of IMessageAttachment. */
-typedef TSharedRef<class IMessageAttachment, ESPMode::ThreadSafe> IMessageAttachmentRef;
-
-
 /**
  * Interface for message attachments.
  *
@@ -43,3 +36,10 @@ public:
 	/** Virtual destructor. */
 	virtual ~IMessageAttachment() { }
 };
+
+
+/** Type definition for shared pointers to instances of IMessageAttachment. */
+typedef TSharedPtr<IMessageAttachment, ESPMode::ThreadSafe> IMessageAttachmentPtr;
+
+/** Type definition for shared references to instances of IMessageAttachment. */
+typedef TSharedRef<IMessageAttachment, ESPMode::ThreadSafe> IMessageAttachmentRef;
diff --git a/Engine/Source/Runtime/Messaging/Public/Interfaces/IMessageBridge.h b/Engine/Source/Runtime/Messaging/Public/Interfaces/IMessageBridge.h
--- a/Engine/Source/Runtime/Messaging/Public/Interfaces/IMessageBridge.h
+++ b/Engine/Source/Runtime/Messaging/Public/Interfaces/IMessageBridge.h
@@ -3,13 +3,6 @@
 #pragma once
 
 
-/** Type definition for shared pointers to instances of IMessageBridge. */
-typedef TSharedPtr<class IMessageBridge, ESPMode::ThreadSafe> IMessageBridgePtr;
-
-/** Type definition for shared references to instances of IMessageBridge. */
-typedef TSharedRef<class IMessageBridge, ESPMode::ThreadSafe> IMessageBridgeRef;
-
-
 /**
  * Interface for message bridges.
  *
@@ -61,3 +54,10 @@ public:
 	/** Virtual destructor. */
 	virtual ~IMessageBridge() { }
 };
+
+
+/** Type definition for shared pointers to instances of IMessageBridge. */
+typedef TSharedPtr<IMessageBridge, ESPMode::ThreadSafe> IMessageBridgePtr;
+
+/** Type definition for shared references to instances of IMessageBridge. */
+typedef TSharedRef<IMessageBridge, ESPMode::ThreadSafe> IMessageBridgeRef;
diff --git a/Engine/Source/Runtime/Messaging/Public/Interfaces/IMessageBus.h b/Engine/Source/Runtime/Messaging/Public/Interfaces/IMessageBus.h
--- a/Engine/Source/Runtime/Messaging/Public/Interfaces/IMessageBus.h
+++ b/Engine/Source/Runtime/Messaging/Public/Interfaces/IMessageBus.h
@@ -3,19 +3,18 @@
 #pragma once
 
 
-/** Type definition for shared pointers to instances of IMessageBus. */
-typedef TSharedPtr<class IMessageBus, ESPMode::ThreadSafe> IMessageBusPtr;
-
-/** Type definition for shared references to instances of IMessageBus. */
-typedef TSharedRef<class IMessageBus, ESPMode::ThreadSafe> IMessageBusRef;
-
-/** Type definition for weak pointers to instances of IMessageBus. */
-typedef TWeakPtr<class IMessageBus, ESPMode::ThreadSafe> IMessageBusWeakPtr;
-
-
-/**
- * Delegate type for message bus shutdowns.
- */
+enum class EMessageScope;
+struct FMessageAddress;
+class IInterceptMessages;
+class IMessageAttachment;
+class IMessageContext;
+class IMessageSubscription;
+class IMessageTracer;
+class IReceiveMessages;
+class ISendMessages;
+
+
+/** Delegate type for message bus shutdowns. */
 DECLARE_MULTICAST_DELEGATE(FOnMessageBusShutdown);
 
 
@@ -76,7 +75,7 @@ DECLARE_MULTICAST_DELEGATE(FOnMessageBusShutdown);
  *
  * It is possible to intercept messages before they are being routed to recipients by registering so called Message Interceptors (objects
  * implementing the IInterceptMessage interface) with the message bus. This feature enables advanced use cases that require inspection or
- * manipulation of messages contents, such as message filtering and enriching, splitting and aggregating, resquencing or authentication[4].
+ * manipulation of messages contents, such as message filtering and enriching, splitting and aggregating, resequencing or authentication[4].
  *
  * The messaging system also provides a facility for debugging the system itself through the so called Message Tracer (an internal object
  * implementing the @see IMessageTracer interface) that can be accessed with the IMessageBus.GetTracer() method. The message tracer is
@@ -114,14 +113,14 @@ public:
 	 * @param Forwarder The sender that forwards the message.
 	 * @see Publish, Send
 	 */
-	virtual void Forward( const IMessageContextRef& Context, const TArray<FMessageAddress>& Recipients, EMessageScope NewScope, const FTimespan& Delay, const ISendMessagesRef& Forwarder ) = 0;
+	virtual void Forward( const TSharedRef<IMessageContext, ESPMode::ThreadSafe>& Context, const TArray<FMessageAddress>& Recipients, EMessageScope NewScope, const FTimespan& Delay, const TSharedRef<ISendMessages, ESPMode::ThreadSafe>& Forwarder ) = 0;
 
 	/**
 	 * Gets the message bus tracer.
 	 *
 	 * @return Weak pointer to the tracer object.
 	 */
-	virtual IMessageTracerRef GetTracer() = 0;
+	virtual TSharedRef<IMessageTracer, ESPMode::ThreadSafe> GetTracer() = 0;
 
 	/**
 	 * Adds an interceptor for messages of the specified type.
@@ -130,7 +129,7 @@ public:
 	 * @param MessageType The type of messages to intercept.
 	 * @see Unintercept
 	 */
-	virtual void Intercept( const IInterceptMessagesRef& Interceptor, const FName& MessageType ) = 0;
+	virtual void Intercept( const TSharedRef<IInterceptMessages, ESPMode::ThreadSafe>& Interceptor, const FName& MessageType ) = 0;
 
 	/**
 	 * Sends a message to subscribed recipients.
@@ -145,7 +144,7 @@ public:
 	 * @param Publisher The message publisher.
 	 * @see Forward, Send
 	 */
-	virtual void Publish( void* Message, UScriptStruct* TypeInfo, EMessageScope Scope, const FTimespan& Delay, const FDateTime& Expiration, const ISendMessagesRef& Publisher ) = 0;
+	virtual void Publish( void* Message, UScriptStruct* TypeInfo, EMessageScope Scope, const FTimespan& Delay, const FDateTime& Expiration, const TSharedRef<ISendMessages, ESPMode::ThreadSafe>& Publisher ) = 0;
 
 	/**
 	 * Registers a message recipient with the message bus.
@@ -154,7 +153,7 @@ public:
 	 * @param Recipient The message recipient.
 	 * @see Unregister
 	 */
-	virtual void Register( const FMessageAddress& Address, const IReceiveMessagesRef& Recipient) = 0;
+	virtual void Register( const FMessageAddress& Address, const TSharedRef<IReceiveMessages, ESPMode::ThreadSafe>& Recipient) = 0;
 
 	/**
 	 * Sends a message to multiple recipients.
@@ -168,7 +167,7 @@ public:
 	 * @param Sender The message sender.
 	 * @see Forward, Publish
 	 */
-	virtual void Send( void* Message, UScriptStruct* TypeInfo, const IMessageAttachmentPtr& Attachment, const TArray<FMessageAddress>& Recipients, const FTimespan& Delay, const FDateTime& Expiration, const ISendMessagesRef& Sender ) = 0;
+	virtual void Send( void* Message, UScriptStruct* TypeInfo, const TSharedPtr<IMessageAttachment, ESPMode::ThreadSafe>& Attachment, const TArray<FMessageAddress>& Recipients, const FTimespan& Delay, const FDateTime& Expiration, const TSharedRef<ISendMessages, ESPMode::ThreadSafe>& Sender ) = 0;
 
 	/**
 	 * Shuts down the message bus.
@@ -189,7 +188,7 @@ public:
 	 * @return The added subscription, or nullptr if the subscription failed.
 	 * @see Unsubscribe
 	 */
-	virtual IMessageSubscriptionPtr Subscribe( const IReceiveMessagesRef& Subscriber, const FName& MessageType, const FMessageScopeRange& ScopeRange ) = 0;
+	virtual TSharedPtr<IMessageSubscription, ESPMode::ThreadSafe> Subscribe( const TSharedRef<IReceiveMessages, ESPMode::ThreadSafe>& Subscriber, const FName& MessageType, const TRange<EMessageScope>& ScopeRange ) = 0;
 
 	/**
 	 * Removes an interceptor for messages of the specified type.
@@ -198,7 +197,7 @@ public:
 	 * @param MessageType The type of messages to stop intercepting.
 	 * @see Intercept
 	 */
-	virtual void Unintercept( const IInterceptMessagesRef& Interceptor, const FName& MessageType ) = 0;
+	virtual void Unintercept( const TSharedRef<IInterceptMessages, ESPMode::ThreadSafe>& Interceptor, const FName& MessageType ) = 0;
 
 	/**
 	 * Unregisters a message recipient from the message bus.
@@ -215,7 +214,7 @@ public:
 	 * @param MessageType The type of messages to unsubscribe from (NAME_All = all types).
 	 * @see Subscribe
 	 */
-	virtual void Unsubscribe( const IReceiveMessagesRef& Subscriber, const FName& MessageType ) = 0;
+	virtual void Unsubscribe( const TSharedRef<IReceiveMessages, ESPMode::ThreadSafe>& Subscriber, const FName& MessageType ) = 0;
 
 public:
 
@@ -232,3 +231,13 @@ public:
 	/** Virtual destructor. */
 	virtual ~IMessageBus() { }
 };
+
+
+/** Type definition for shared pointers to instances of IMessageBus. */
+typedef TSharedPtr<IMessageBus, ESPMode::ThreadSafe> IMessageBusPtr;
+
+/** Type definition for shared references to instances of IMessageBus. */
+typedef TSharedRef<IMessageBus, ESPMode::ThreadSafe> IMessageBusRef;
+
+/** Type definition for weak pointers to instances of IMessageBus. */
+typedef TWeakPtr<IMessageBus, ESPMode::ThreadSafe> IMessageBusWeakPtr;
diff --git a/Engine/Source/Runtime/Messaging/Public/Interfaces/IMessageContext.h b/Engine/Source/Runtime/Messaging/Public/Interfaces/IMessageContext.h
--- a/Engine/Source/Runtime/Messaging/Public/Interfaces/IMessageContext.h
+++ b/Engine/Source/Runtime/Messaging/Public/Interfaces/IMessageContext.h
@@ -3,6 +3,124 @@
 #pragma once
 
 
+class IMessageAttachment;
+
+
+/**
+ * Implements a message endpoint address.
+ */
+struct FMessageAddress
+{
+public:
+
+	/** Default constructor (no initialization). */
+	FMessageAddress() { }
+
+public:
+
+	/**
+	 * Compares two message addresses for equality.
+	 *
+	 * @param X The first address to compare.
+	 * @param Y The second address to compare.
+	 * @return true if the addresses are equal, false otherwise.
+	 */
+	friend bool operator==( const FMessageAddress& X, const FMessageAddress& Y )
+	{
+		return (X.UniqueId == Y.UniqueId);
+	}
+
+	/**
+	 * Compares two message addresses for inequality.
+	 *
+	 * @param X The first address to compare.
+	 * @param Y The second address to compare.
+	 * @return true if the addresses are not equal, false otherwise.
+	 */
+	friend bool operator!=( const FMessageAddress& X, const FMessageAddress& Y )
+	{
+		return (X.UniqueId != Y.UniqueId);
+	}
+
+	/**
+	 * Serializes a message address from or into an archive.
+	 *
+	 * @param Ar The archive to serialize from or into.
+	 * @param G The address to serialize.
+	 */
+	friend FArchive& operator<<( FArchive& Ar, FMessageAddress& A )
+	{
+		return Ar << A.UniqueId;
+	}
+
+public:
+
+	/**
+	 * Invalidates the GUID.
+	 *
+	 * @see IsValid
+	 */
+	void Invalidate()
+	{
+		UniqueId.Invalidate();
+	}
+
+	/**
+	 * Checks whether this message address is valid or not.
+	 *
+	 * @return true if valid, false otherwise.
+	 * @see Invalidate
+	 */
+	bool IsValid() const
+	{
+		return UniqueId.IsValid();
+	}
+
+	/**
+	 * Converts this GUID to its string representation.
+	 *
+	 * @return The string representation.
+	 */
+	FString ToString() const
+	{
+		return UniqueId.ToString();
+	}
+
+public:
+
+	/**
+	 * Calculates the hash for a message address.
+	 *
+	 * @param Address The address to calculate the hash for.
+	 * @return The hash.
+	 */
+	friend uint32 GetTypeHash( const FMessageAddress& Address )
+	{
+		return FCrc::MemCrc_DEPRECATED(&Address.UniqueId, sizeof(FGuid));
+	}
+
+public:
+
+	/**
+	 * Returns a new message address.
+	 *
+	 * @return A new address.
+	 */
+	static FMessageAddress NewAddress()
+	{
+		FMessageAddress Result;
+		Result.UniqueId = FGuid::NewGuid();
+
+		return Result;
+	}
+
+private:
+
+	/** Holds a unique identifier. */
+	FGuid UniqueId;
+};
+
+
 /**
  * Enumerates scopes for published messages.
  *
@@ -38,15 +156,6 @@ enum class EMessageScope
 };
 
 
-/** Type definition for message endpoint identifiers. */
-typedef FGuid FMessageAddress;
-
-/** Type definition for shared pointers to instances of IMessageContext. */
-typedef TSharedPtr<class IMessageContext, ESPMode::ThreadSafe> IMessageContextPtr;
-
-/** Type definition for shared references to instances of IMessageContext. */
-typedef TSharedRef<class IMessageContext, ESPMode::ThreadSafe> IMessageContextRef;
-
 /** Type definition for message scope ranges. */
 typedef TRange<EMessageScope> FMessageScopeRange;
 
@@ -80,7 +189,7 @@ public:
 	 *
 	 * @return A pointer to the message attachment, or nullptr if no attachment is present.
 	 */
-	virtual IMessageAttachmentPtr GetAttachment() const = 0;
+	virtual TSharedPtr<IMessageAttachment, ESPMode::ThreadSafe> GetAttachment() const = 0;
 
 	/**
 	 * Gets the date and time at which the message expires.
@@ -130,7 +239,7 @@ public:
 	 *
 	 * @return The original message.
 	 */
-	virtual IMessageContextPtr GetOriginalContext() const = 0;
+	virtual TSharedPtr<IMessageContext, ESPMode::ThreadSafe> GetOriginalContext() const = 0;
 
 	/**
 	* Gets the list of message recipients.
@@ -194,3 +303,10 @@ public:
 	/** Virtual destructor. */
 	virtual ~IMessageContext() { }
 };
+
+
+/** Type definition for shared pointers to instances of IMessageContext. */
+typedef TSharedPtr<IMessageContext, ESPMode::ThreadSafe> IMessageContextPtr;
+
+/** Type definition for shared references to instances of IMessageContext. */
+typedef TSharedRef<IMessageContext, ESPMode::ThreadSafe> IMessageContextRef;
diff --git a/Engine/Source/Runtime/Messaging/Public/Interfaces/IMessageData.h b/Engine/Source/Runtime/Messaging/Public/Interfaces/IMessageData.h
--- a/Engine/Source/Runtime/Messaging/Public/Interfaces/IMessageData.h
+++ b/Engine/Source/Runtime/Messaging/Public/Interfaces/IMessageData.h
@@ -19,13 +19,6 @@ enum class EMessageDataState
 };
 
 
-/** Type definition for shared pointers to instances of IMessageData. */
-typedef TSharedPtr<class IMessageData, ESPMode::ThreadSafe> IMessageDataPtr;
-
-/** Type definition for shared references to instances of IMessageData. */
-typedef TSharedRef<class IMessageData, ESPMode::ThreadSafe> IMessageDataRef;
-
-
 /**
  * Interface for serialized message data.
  */
@@ -65,3 +58,10 @@ public:
 	/** Virtual destructor. */
 	virtual ~IMessageData() { }
 };
+
+
+/** Type definition for shared pointers to instances of IMessageData. */
+typedef TSharedPtr<IMessageData, ESPMode::ThreadSafe> IMessageDataPtr;
+
+/** Type definition for shared references to instances of IMessageData. */
+typedef TSharedRef<IMessageData, ESPMode::ThreadSafe> IMessageDataRef;
diff --git a/Engine/Source/Runtime/Messaging/Public/Interfaces/IMessageHandler.h b/Engine/Source/Runtime/Messaging/Public/Interfaces/IMessageHandler.h
--- a/Engine/Source/Runtime/Messaging/Public/Interfaces/IMessageHandler.h
+++ b/Engine/Source/Runtime/Messaging/Public/Interfaces/IMessageHandler.h
@@ -3,11 +3,7 @@
 #pragma once
 
 
-/** Type definition for shared pointers to instances of IMessageHandler. */
-typedef TSharedPtr<class IMessageHandler, ESPMode::ThreadSafe> IMessageHandlerPtr;
-
-/** Type definition for shared references to instances of IMessageHandler. */
-typedef TSharedRef<class IMessageHandler, ESPMode::ThreadSafe> IMessageHandlerRef;
+class IMessageContext;
 
 
 /**
@@ -21,6 +17,7 @@ public:
 	 * Gets the message type handled by this handler.
 	 *
 	 * @return The name of the message type.
+	 * @see HandleMessage
 	 */
 	virtual const FName GetHandledMessageType() const = 0;
 
@@ -28,8 +25,9 @@ public:
 	 * Handles the specified message.
 	 *
 	 * @param Context The context of the message to handle.
+	 * @see GetHandledMessageType
 	 */
-	virtual void HandleMessage( const IMessageContextRef& Context ) = 0;
+	virtual void HandleMessage( const TSharedRef<IMessageContext, ESPMode::ThreadSafe>& Context ) = 0;
 
 public:
 
@@ -38,6 +36,13 @@ public:
 };
 
 
+/** Type definition for shared pointers to instances of IMessageHandler. */
+typedef TSharedPtr<IMessageHandler, ESPMode::ThreadSafe> IMessageHandlerPtr;
+
+/** Type definition for shared references to instances of IMessageHandler. */
+typedef TSharedRef<IMessageHandler, ESPMode::ThreadSafe> IMessageHandlerRef;
+
+
 /**
  * Template for message handler function pointer types.
  *
@@ -47,7 +52,7 @@ public:
 template<typename MessageType, typename HandlerType>
 struct TMessageHandlerFunc
 {
-	typedef void (HandlerType::*Type)( const MessageType&, const IMessageContextRef& );
+	typedef void (HandlerType::*Type)( const MessageType&, const TSharedRef<IMessageContext, ESPMode::ThreadSafe>& );
 };
 
 
@@ -85,7 +90,7 @@ public:
 		return MessageType::StaticStruct()->GetFName();
 	}
 
-	virtual void HandleMessage( const IMessageContextRef& Context ) override
+	virtual void HandleMessage( const TSharedRef<IMessageContext, ESPMode::ThreadSafe>& Context ) override
 	{
 		(Handler->*HandlerFunc)(*static_cast<const MessageType*>(Context->GetMessage()), Context);
 	}
diff --git a/Engine/Source/Runtime/Messaging/Public/Interfaces/IMessageSubscription.h b/Engine/Source/Runtime/Messaging/Public/Interfaces/IMessageSubscription.h
--- a/Engine/Source/Runtime/Messaging/Public/Interfaces/IMessageSubscription.h
+++ b/Engine/Source/Runtime/Messaging/Public/Interfaces/IMessageSubscription.h
@@ -3,11 +3,8 @@
 #pragma once
 
 
-/** Type definition for shared pointers to instances of IMessageSubscription. */
-typedef TSharedPtr<class IMessageSubscription, ESPMode::ThreadSafe> IMessageSubscriptionPtr;
-
-/** Type definition for shared references to instances of IMessageSubscription. */
-typedef TSharedRef<class IMessageSubscription, ESPMode::ThreadSafe> IMessageSubscriptionRef;
+enum class EMessageScope;
+class IReceiveMessages;
 
 
 /**
@@ -53,7 +50,7 @@ public:
 	 * @return Message scope range.
 	 * @see GetMessageType, GetSubscriber
 	 */
-	virtual const FMessageScopeRange& GetScopeRange() = 0;
+	virtual const TRange<EMessageScope>& GetScopeRange() = 0;
 
 	/**
 	 * Gets the subscriber.
@@ -61,7 +58,7 @@ public:
 	 * @return The subscriber.
 	 * @see GetMessageType, GetScopeRange
 	 */
-	virtual const IReceiveMessagesWeakPtr& GetSubscriber() = 0;
+	virtual const TWeakPtr<IReceiveMessages, ESPMode::ThreadSafe>& GetSubscriber() = 0;
 
 	/**
 	 * Checks whether the subscription is enabled.
@@ -76,3 +73,10 @@ public:
 	/** Virtual destructor. */
 	virtual ~IMessageSubscription() { }
 };
+
+
+/** Type definition for shared pointers to instances of IMessageSubscription. */
+typedef TSharedPtr<IMessageSubscription, ESPMode::ThreadSafe> IMessageSubscriptionPtr;
+
+/** Type definition for shared references to instances of IMessageSubscription. */
+typedef TSharedRef<IMessageSubscription, ESPMode::ThreadSafe> IMessageSubscriptionRef;
diff --git a/Engine/Source/Runtime/Messaging/Public/Interfaces/IMessageTracer.h b/Engine/Source/Runtime/Messaging/Public/Interfaces/IMessageTracer.h
--- a/Engine/Source/Runtime/Messaging/Public/Interfaces/IMessageTracer.h
+++ b/Engine/Source/Runtime/Messaging/Public/Interfaces/IMessageTracer.h
@@ -2,12 +2,8 @@
 
 #pragma once
 
+#include "IMessageContext.h"
 
-/** Type definition for shared pointers to instances of IMessageTracer. */
-typedef TSharedPtr<class IMessageTracer, ESPMode::ThreadSafe> IMessageTracerPtr;
-
-/** Type definition for shared references to instances of IMessageTracer. */
-typedef TSharedRef<class IMessageTracer, ESPMode::ThreadSafe> IMessageTracerRef;
 
 /** Type definition for shared pointers to instances of FMessageTracerAddressInfo. */
 typedef TSharedPtr<struct FMessageTracerAddressInfo> FMessageTracerAddressInfoPtr;
@@ -304,3 +300,10 @@ protected:
 	/** Hidden destructor. The life time of a message tracer is managed by the message bus. */
 	virtual ~IMessageTracer() { }
 };
+
+
+/** Type definition for shared pointers to instances of IMessageTracer. */
+typedef TSharedPtr<IMessageTracer, ESPMode::ThreadSafe> IMessageTracerPtr;
+
+/** Type definition for shared references to instances of IMessageTracer. */
+typedef TSharedRef<IMessageTracer, ESPMode::ThreadSafe> IMessageTracerRef;
diff --git a/Engine/Source/Runtime/Messaging/Public/Interfaces/IMessageTracerBreakpoint.h b/Engine/Source/Runtime/Messaging/Public/Interfaces/IMessageTracerBreakpoint.h
--- a/Engine/Source/Runtime/Messaging/Public/Interfaces/IMessageTracerBreakpoint.h
+++ b/Engine/Source/Runtime/Messaging/Public/Interfaces/IMessageTracerBreakpoint.h
@@ -3,11 +3,7 @@
 #pragma once
 
 
-/** Type definition for shared pointers to instances of IMessageTracerBreakpoint. */
-typedef TSharedPtr<class IMessageTracerBreakpoint, ESPMode::ThreadSafe> IMessageTracerBreakpointPtr;
-
-/** Type definition for shared references to instances of IMessageTracerBreakpoint. */
-typedef TSharedRef<class IMessageTracerBreakpoint, ESPMode::ThreadSafe> IMessageTracerBreakpointRef;
+class IMessageContext;
 
 
 /**
@@ -32,10 +28,17 @@ public:
 	 * @param Context The message to break on.
 	 * @return true if the tracer should break, false otherwise.
 	 */
-	virtual bool ShouldBreak( const IMessageContextRef& Context ) const = 0;
+	virtual bool ShouldBreak( const TSharedRef<IMessageContext, ESPMode::ThreadSafe>& Context ) const = 0;
 
 protected:
 
 	/** Hidden destructor. */
 	virtual ~IMessageTracerBreakpoint() { }
 };
+
+
+/** Type definition for shared pointers to instances of IMessageTracerBreakpoint. */
+typedef TSharedPtr<IMessageTracerBreakpoint, ESPMode::ThreadSafe> IMessageTracerBreakpointPtr;
+
+/** Type definition for shared references to instances of IMessageTracerBreakpoint. */
+typedef TSharedRef<IMessageTracerBreakpoint, ESPMode::ThreadSafe> IMessageTracerBreakpointRef;
diff --git a/Engine/Source/Runtime/Messaging/Public/Interfaces/IMessagingModule.h b/Engine/Source/Runtime/Messaging/Public/Interfaces/IMessagingModule.h
--- a/Engine/Source/Runtime/Messaging/Public/Interfaces/IMessagingModule.h
+++ b/Engine/Source/Runtime/Messaging/Public/Interfaces/IMessagingModule.h
@@ -2,6 +2,16 @@
 
 #pragma once
 
+#include "ModuleInterface.h"
+
+
+struct FMessageAddress;
+class IAuthorizeMessageRecipients;
+class IMessageBus;
+class IMessageBridge;
+class ISerializeMessages;
+class ITransportMessages;
+
 
 /**
  * Interface for messaging modules.
@@ -26,7 +36,7 @@ public:
 	 * @return The new message bridge, or nullptr if the bridge couldn't be created.
 	 * @see CreateBus
 	 */
-	virtual IMessageBridgePtr CreateBridge( const FMessageAddress& Address, const IMessageBusRef& Bus, const ISerializeMessagesRef& Serializer, const ITransportMessagesRef& Transport ) = 0;
+	virtual TSharedPtr<IMessageBridge, ESPMode::ThreadSafe> CreateBridge( const FMessageAddress& Address, const TSharedRef<IMessageBus, ESPMode::ThreadSafe>& Bus, const TSharedRef<ISerializeMessages, ESPMode::ThreadSafe>& Serializer, const TSharedRef<ITransportMessages, ESPMode::ThreadSafe>& Transport ) = 0;
 
 	/**
 	 * Creates a new message bus.
@@ -35,30 +45,29 @@ public:
 	 * @return The new message bus, or nullptr if the bus couldn't be created.
 	 * @see CreateBridge
 	 */
-	virtual IMessageBusPtr CreateBus( const IAuthorizeMessageRecipientsPtr& RecipientAuthorizer ) = 0;
+	virtual TSharedPtr<IMessageBus, ESPMode::ThreadSafe> CreateBus( const TSharedPtr<IAuthorizeMessageRecipients>& RecipientAuthorizer ) = 0;
 
 	/**
 	 * Creates a Json message serializer (deprecated).
 	 *
 	 * @return A new serializer.
 	 */
-	virtual ISerializeMessagesPtr CreateJsonMessageSerializer() = 0;
+	virtual TSharedPtr<ISerializeMessages, ESPMode::ThreadSafe> CreateJsonMessageSerializer() = 0;
 
 	/**
 	 * Gets the default message bus if it has been initialized.
 	 *
 	 * @return The default bus.
 	 */
-	virtual IMessageBusPtr GetDefaultBus() const = 0;
+	virtual TSharedPtr<IMessageBus, ESPMode::ThreadSafe> GetDefaultBus() const = 0;
 
 public:
 
 	/**
 	 * Gets a reference to the messaging module instance.
 	 *
-	 * @todo gmp: better implementation using dependency injection.
-	 *
 	 * @return A reference to the Messaging module.
+	 * @todo gmp: better implementation using dependency injection.
 	 */
 	static IMessagingModule& Get()
 	{
diff --git a/Engine/Source/Runtime/Messaging/Public/Interfaces/IMutableMessageContext.h b/Engine/Source/Runtime/Messaging/Public/Interfaces/IMutableMessageContext.h
--- a/Engine/Source/Runtime/Messaging/Public/Interfaces/IMutableMessageContext.h
+++ b/Engine/Source/Runtime/Messaging/Public/Interfaces/IMutableMessageContext.h
@@ -2,12 +2,11 @@
 
 #pragma once
 
+#include "IMessageContext.h"
 
-/** Type definition for shared pointers to instances of IMutableMessageContext. */
-typedef TSharedPtr<class IMutableMessageContext, ESPMode::ThreadSafe> IMutableMessageContextPtr;
 
-/** Type definition for shared references to instances of IMutableMessageContext. */
-typedef TSharedRef<class IMutableMessageContext, ESPMode::ThreadSafe> IMutableMessageContextRef;
+class IMessageAttachment;
+class UScriptStruct;
 
 
 /**
@@ -30,7 +29,7 @@ public:
 	 *
 	 * @param InAttachment The attachment to set.
 	 */
-	virtual void SetAttachment( const IMessageAttachmentPtr& InAttachment ) = 0;
+	virtual void SetAttachment( const TSharedPtr<IMessageAttachment, ESPMode::ThreadSafe>& InAttachment ) = 0;
 
 	/**
 	 * Sets the message.
@@ -81,3 +80,10 @@ public:
 	/** Virtual destructor. */
 	virtual ~IMutableMessageContext() { }
 };
+
+
+/** Type definition for shared pointers to instances of IMutableMessageContext. */
+typedef TSharedPtr<IMutableMessageContext, ESPMode::ThreadSafe> IMutableMessageContextPtr;
+
+/** Type definition for shared references to instances of IMutableMessageContext. */
+typedef TSharedRef<IMutableMessageContext, ESPMode::ThreadSafe> IMutableMessageContextRef;
diff --git a/Engine/Source/Runtime/Messaging/Public/Interfaces/IReceiveMessages.h b/Engine/Source/Runtime/Messaging/Public/Interfaces/IReceiveMessages.h
--- a/Engine/Source/Runtime/Messaging/Public/Interfaces/IReceiveMessages.h
+++ b/Engine/Source/Runtime/Messaging/Public/Interfaces/IReceiveMessages.h
@@ -3,14 +3,7 @@
 #pragma once
 
 
-/** Type definition for shared pointers to instances of IReceiveMessages. */
-typedef TSharedPtr<class IReceiveMessages, ESPMode::ThreadSafe> IReceiveMessagesPtr;
-
-/** Type definition for shared references to instances of IReceiveMessages. */
-typedef TSharedRef<class IReceiveMessages, ESPMode::ThreadSafe> IReceiveMessagesRef;
-
-/** Type definition for shared pointers to instances of IReceiveMessages. */
-typedef TWeakPtr<class IReceiveMessages, ESPMode::ThreadSafe> IReceiveMessagesWeakPtr;
+class IMessageContext;
 
 
 /**
@@ -71,7 +64,7 @@ public:
 	 *
 	 * @param Context Will hold the context of the received message.
 	 */
-	virtual void ReceiveMessage( const IMessageContextRef& Context ) = 0;
+	virtual void ReceiveMessage( const TSharedRef<IMessageContext, ESPMode::ThreadSafe>& Context ) = 0;
 
 public:
 
@@ -94,3 +87,13 @@ public:
 	/** Virtual destructor. */
 	virtual ~IReceiveMessages() { }
 };
+
+
+/** Type definition for shared pointers to instances of IReceiveMessages. */
+typedef TSharedPtr<IReceiveMessages, ESPMode::ThreadSafe> IReceiveMessagesPtr;
+
+/** Type definition for shared references to instances of IReceiveMessages. */
+typedef TSharedRef<IReceiveMessages, ESPMode::ThreadSafe> IReceiveMessagesRef;
+
+/** Type definition for shared pointers to instances of IReceiveMessages. */
+typedef TWeakPtr<IReceiveMessages, ESPMode::ThreadSafe> IReceiveMessagesWeakPtr;
diff --git a/Engine/Source/Runtime/Messaging/Public/Interfaces/ISendMessages.h b/Engine/Source/Runtime/Messaging/Public/Interfaces/ISendMessages.h
--- a/Engine/Source/Runtime/Messaging/Public/Interfaces/ISendMessages.h
+++ b/Engine/Source/Runtime/Messaging/Public/Interfaces/ISendMessages.h
@@ -3,11 +3,7 @@
 #pragma once
 
 
-/** Type definition for shared pointers to instances of ISendMessages. */
-typedef TSharedPtr<class ISendMessages, ESPMode::ThreadSafe> ISendMessagesPtr;
-
-/** Type definition for shared references to instances of ISendMessages. */
-typedef TSharedRef<class ISendMessages, ESPMode::ThreadSafe> ISendMessagesRef;
+class IMessageContext;
 
 
 /**
@@ -43,10 +39,17 @@ public:
 	 * @param Context The context of the message that generated the error.
 	 * @param Error The error string.
 	 */
-	virtual void NotifyMessageError( const IMessageContextRef& Context, const FString& Error ) = 0;
+	virtual void NotifyMessageError( const TSharedRef<IMessageContext, ESPMode::ThreadSafe>& Context, const FString& Error ) = 0;
 
 public:
 
 	/** Virtual destructor. */
 	virtual ~ISendMessages() { }
 };
+
+
+/** Type definition for shared pointers to instances of ISendMessages. */
+typedef TSharedPtr<ISendMessages, ESPMode::ThreadSafe> ISendMessagesPtr;
+
+/** Type definition for shared references to instances of ISendMessages. */
+typedef TSharedRef<ISendMessages, ESPMode::ThreadSafe> ISendMessagesRef;
diff --git a/Engine/Source/Runtime/Messaging/Public/Interfaces/ISerializeMessages.h b/Engine/Source/Runtime/Messaging/Public/Interfaces/ISerializeMessages.h
--- a/Engine/Source/Runtime/Messaging/Public/Interfaces/ISerializeMessages.h
+++ b/Engine/Source/Runtime/Messaging/Public/Interfaces/ISerializeMessages.h
@@ -3,11 +3,8 @@
 #pragma once
 
 
-/** Type definition for shared pointers to instances of ISerializeMessages. */
-typedef TSharedPtr<class ISerializeMessages, ESPMode::ThreadSafe> ISerializeMessagesPtr;
-
-/** Type definition for shared references to instances of ISerializeMessages. */
-typedef TSharedRef<class ISerializeMessages, ESPMode::ThreadSafe> ISerializeMessagesRef;
+class IMessageContext;
+class IMutableMessageContext;
 
 
 /**
@@ -27,7 +24,7 @@ public:
 	 * @param OutContext Will hold the context of the deserialized message.
 	 * @return true if deserialization was successful, false otherwise.
 	 */
-	virtual bool DeserializeMessage( FArchive& Archive, IMutableMessageContextRef& OutContext ) = 0;
+	virtual bool DeserializeMessage( FArchive& Archive, TSharedRef<IMutableMessageContext, ESPMode::ThreadSafe>& OutContext ) = 0;
 
 	/**
 	 * Serializes a message into an archive.
@@ -36,10 +33,17 @@ public:
 	 * @param Archive The archive to serialize into.
 	 * @return true if serialization was successful, false otherwise.
 	 */
-	virtual bool SerializeMessage( const IMessageContextRef& Context, FArchive& Archive ) = 0;
+	virtual bool SerializeMessage( const TSharedRef<IMessageContext, ESPMode::ThreadSafe>& Context, FArchive& Archive ) = 0;
 
 public:
 
 	/** Virtual destructor. */
 	virtual ~ISerializeMessages() { }
 };
+
+
+/** Type definition for shared pointers to instances of ISerializeMessages. */
+typedef TSharedPtr<ISerializeMessages, ESPMode::ThreadSafe> ISerializeMessagesPtr;
+
+/** Type definition for shared references to instances of ISerializeMessages. */
+typedef TSharedRef<ISerializeMessages, ESPMode::ThreadSafe> ISerializeMessagesRef;
diff --git a/Engine/Source/Runtime/Messaging/Public/Interfaces/ITransportMessages.h b/Engine/Source/Runtime/Messaging/Public/Interfaces/ITransportMessages.h
--- a/Engine/Source/Runtime/Messaging/Public/Interfaces/ITransportMessages.h
+++ b/Engine/Source/Runtime/Messaging/Public/Interfaces/ITransportMessages.h
@@ -3,11 +3,8 @@
 #pragma once
 
 
-/** Type definition for shared pointers to instances of ITransportMessages. */
-typedef TSharedPtr<class ITransportMessages, ESPMode::ThreadSafe> ITransportMessagesPtr;
-
-/** Type definition for shared references to instances of ITransportMessages. */
-typedef TSharedRef<class ITransportMessages, ESPMode::ThreadSafe> ITransportMessagesRef;
+class IMessageAttachment;
+class IMessageData;
 
 
 /**
@@ -77,7 +74,7 @@ public:
 	 * @param Recipients The transport nodes to send the message to.
 	 * @return true if the message is being transported, false otherwise.
 	 */
-	virtual bool TransportMessage( const IMessageDataRef& Data, const IMessageAttachmentPtr& Attachment, const TArray<FGuid>& Recipients ) = 0;
+	virtual bool TransportMessage( const TSharedRef<IMessageData, ESPMode::ThreadSafe>& Data, const TSharedPtr<IMessageAttachment, ESPMode::ThreadSafe>& Attachment, const TArray<FGuid>& Recipients ) = 0;
 
 public:
 
@@ -107,3 +104,10 @@ protected:
 	/** Virtual constructor. */
 	~ITransportMessages() { }
 };
+
+
+/** Type definition for shared pointers to instances of ITransportMessages. */
+typedef TSharedPtr<ITransportMessages, ESPMode::ThreadSafe> ITransportMessagesPtr;
+
+/** Type definition for shared references to instances of ITransportMessages. */
+typedef TSharedRef<ITransportMessages, ESPMode::ThreadSafe> ITransportMessagesRef;
