Messaging: Renamed message transport and interceptor interfaces

#UpgradeNodes: Rename ITransportMessages to IMessageTransport, IInterceptMessages to IMessageInterceptor

diff --git a/Engine/Plugins/Messaging/UdpMessaging/Source/UdpMessaging/Private/Transport/UdpMessageTransport.cpp b/Engine/Plugins/Messaging/UdpMessaging/Source/UdpMessaging/Private/Transport/UdpMessageTransport.cpp
--- a/Engine/Plugins/Messaging/UdpMessaging/Source/UdpMessaging/Private/Transport/UdpMessageTransport.cpp
+++ b/Engine/Plugins/Messaging/UdpMessaging/Source/UdpMessaging/Private/Transport/UdpMessageTransport.cpp
@@ -26,7 +26,7 @@ FUdpMessageTransport::~FUdpMessageTransport()
 }
 
 
-/* ITransportMessages interface
+/* IMessageTransport interface
  *****************************************************************************/
 
 bool FUdpMessageTransport::StartTransport()
diff --git a/Engine/Plugins/Messaging/UdpMessaging/Source/UdpMessaging/Private/Transport/UdpMessageTransport.h b/Engine/Plugins/Messaging/UdpMessaging/Source/UdpMessaging/Private/Transport/UdpMessageTransport.h
--- a/Engine/Plugins/Messaging/UdpMessaging/Source/UdpMessaging/Private/Transport/UdpMessageTransport.h
+++ b/Engine/Plugins/Messaging/UdpMessaging/Source/UdpMessaging/Private/Transport/UdpMessageTransport.h
@@ -7,7 +7,7 @@
  * Implements a message transport technology using an UDP network connection.
  */
 class FUdpMessageTransport
-	: public ITransportMessages
+	: public IMessageTransport
 {
 public:
 
@@ -25,7 +25,7 @@ public:
 
 public:
 
-	// ITransportMessages interface
+	// IMessageTransport interface
 
 	virtual FName GetDebugName() const override
 	{
@@ -53,25 +53,25 @@ public:
 
 private:
 
-	// Handles received transport messages. */
+	/** Handles received transport messages. */
 	void HandleProcessorMessageReceived( FArchive& MessageData, const IMessageAttachmentPtr& Attachment, const FGuid& NodeId )
 	{
 		MessageReceivedDelegate.ExecuteIfBound(MessageData, Attachment, NodeId);
 	}
 
-	// Handles discovered transport endpoints. */
+	/** Handles discovered transport endpoints. */
 	void HandleProcessorNodeDiscovered( const FGuid& DiscoveredNodeId )
 	{
 		NodeDiscoveredDelegate.ExecuteIfBound(DiscoveredNodeId);
 	}
 
-	// Handles lost transport endpoints. */
+	/** Handles lost transport endpoints. */
 	void HandleProcessorNodeLost( const FGuid& LostNodeId )
 	{
 		NodeLostDelegate.ExecuteIfBound(LostNodeId);
 	}
 
-	// Handles received socket data. */
+	/** Handles received socket data. */
 	void HandleSocketDataReceived(const FArrayReaderPtr& Data, const FIPv4Endpoint& Sender);
 
 private:
diff --git a/Engine/Source/Runtime/Messaging/Private/Bridge/MessageBridge.cpp b/Engine/Source/Runtime/Messaging/Private/Bridge/MessageBridge.cpp
--- a/Engine/Source/Runtime/Messaging/Private/Bridge/MessageBridge.cpp
+++ b/Engine/Source/Runtime/Messaging/Private/Bridge/MessageBridge.cpp
@@ -6,7 +6,7 @@
 /* FMessageBridge structors
  *****************************************************************************/
 
-FMessageBridge::FMessageBridge( const FMessageAddress InAddress, const IMessageBusRef& InBus, const ISerializeMessagesRef& InSerializer, const ITransportMessagesRef& InTransport )
+FMessageBridge::FMessageBridge( const FMessageAddress InAddress, const IMessageBusRef& InBus, const ISerializeMessagesRef& InSerializer, const IMessageTransportRef& InTransport )
 	: Address(InAddress)
 	, Bus(InBus)
 	, Enabled(false)
diff --git a/Engine/Source/Runtime/Messaging/Private/Bridge/MessageBridge.h b/Engine/Source/Runtime/Messaging/Private/Bridge/MessageBridge.h
--- a/Engine/Source/Runtime/Messaging/Private/Bridge/MessageBridge.h
+++ b/Engine/Source/Runtime/Messaging/Private/Bridge/MessageBridge.h
@@ -22,7 +22,7 @@ public:
 	 * @param InSerializer The message serializer to use.
 	 * @param InTransport The transport mechanism to use.
 	 */
-	FMessageBridge( const FMessageAddress InAddress, const IMessageBusRef& InBus, const ISerializeMessagesRef& InSerializer, const ITransportMessagesRef& InTransport );
+	FMessageBridge( const FMessageAddress InAddress, const IMessageBusRef& InBus, const ISerializeMessagesRef& InSerializer, const IMessageTransportRef& InTransport );
 
 	/** Destructor. */
 	~FMessageBridge();
@@ -116,5 +116,5 @@ private:
 	ISerializeMessagesPtr Serializer;
 
 	/** Holds the message transport object. */
-	ITransportMessagesPtr Transport;
+	IMessageTransportPtr Transport;
 };
diff --git a/Engine/Source/Runtime/Messaging/Private/Bus/MessageBus.h b/Engine/Source/Runtime/Messaging/Private/Bus/MessageBus.h
--- a/Engine/Source/Runtime/Messaging/Private/Bus/MessageBus.h
+++ b/Engine/Source/Runtime/Messaging/Private/Bus/MessageBus.h
@@ -36,7 +36,7 @@ public:
 		return Router->GetTracer();
 	}
 
-	virtual void Intercept( const IInterceptMessagesRef& Interceptor, const FName& MessageType ) override
+	virtual void Intercept( const IMessageInterceptorRef& Interceptor, const FName& MessageType ) override
 	{
 		if (MessageType != NAME_None)
 		{
@@ -85,7 +85,7 @@ public:
 		return nullptr;
 	}
 
-	virtual void Unintercept( const IInterceptMessagesRef& Interceptor, const FName& MessageType ) override
+	virtual void Unintercept( const IMessageInterceptorRef& Interceptor, const FName& MessageType ) override
 	{
 		if (MessageType != NAME_None)
 		{
diff --git a/Engine/Source/Runtime/Messaging/Private/Bus/MessageRouter.cpp b/Engine/Source/Runtime/Messaging/Private/Bus/MessageRouter.cpp
--- a/Engine/Source/Runtime/Messaging/Private/Bus/MessageRouter.cpp
+++ b/Engine/Source/Runtime/Messaging/Private/Bus/MessageRouter.cpp
@@ -214,7 +214,7 @@ void FMessageRouter::ProcessDelayedMessages( )
 /* FMessageRouter callbacks
  *****************************************************************************/
 
-void FMessageRouter::HandleAddInterceptor( IInterceptMessagesRef Interceptor, FName MessageType )
+void FMessageRouter::HandleAddInterceptor( IMessageInterceptorRef Interceptor, FName MessageType )
 {
 	ActiveInterceptors.FindOrAdd(MessageType).AddUnique(Interceptor);
 	Tracer->TraceAddedInterceptor(Interceptor, MessageType);
@@ -240,19 +240,19 @@ void FMessageRouter::HandleAddSubscriber( IMessageSubscriptionRef Subscription )
 }
 
 
-void FMessageRouter::HandleRemoveInterceptor( IInterceptMessagesRef Interceptor, FName MessageType )
+void FMessageRouter::HandleRemoveInterceptor( IMessageInterceptorRef Interceptor, FName MessageType )
 {
 	if (MessageType == NAME_All)
 	{
-		for (TMap<FName, TArray<IInterceptMessagesPtr> >::TIterator It(ActiveInterceptors); It; ++It)
+		for (TMap<FName, TArray<IMessageInterceptorPtr> >::TIterator It(ActiveInterceptors); It; ++It)
 		{
-			TArray<IInterceptMessagesPtr>& Interceptors = It.Value();
+			TArray<IMessageInterceptorPtr>& Interceptors = It.Value();
 			Interceptors.Remove(Interceptor);
 		}
 	}
 	else
 	{
-		TArray<IInterceptMessagesPtr>& Interceptors = ActiveInterceptors.FindOrAdd(MessageType);
+		TArray<IMessageInterceptorPtr>& Interceptors = ActiveInterceptors.FindOrAdd(MessageType);
 		Interceptors.Remove(Interceptor);
 	}
 
@@ -305,9 +305,9 @@ void FMessageRouter::HandleRemoveSubscriber( IReceiveMessagesWeakPtr SubscriberP
 void FMessageRouter::HandleRouteMessage( IMessageContextRef Context )
 {
 	// intercept routing
-	TArray<IInterceptMessagesPtr>& Interceptors = ActiveInterceptors.FindOrAdd(Context->GetMessageType());
+	TArray<IMessageInterceptorPtr>& Interceptors = ActiveInterceptors.FindOrAdd(Context->GetMessageType());
 
-	for (TArray<IInterceptMessagesPtr>::TIterator It(Interceptors); It; ++It)
+	for (TArray<IMessageInterceptorPtr>::TIterator It(Interceptors); It; ++It)
 	{
 		if ((*It)->InterceptMessage(Context))
 		{
diff --git a/Engine/Source/Runtime/Messaging/Private/Bus/MessageRouter.h b/Engine/Source/Runtime/Messaging/Private/Bus/MessageRouter.h
--- a/Engine/Source/Runtime/Messaging/Private/Bus/MessageRouter.h
+++ b/Engine/Source/Runtime/Messaging/Private/Bus/MessageRouter.h
@@ -27,7 +27,7 @@ public:
 	 * @param Interceptor The interceptor to add.
 	 * @param MessageType The type of messages to intercept.
 	 */
-	FORCEINLINE void AddInterceptor( const IInterceptMessagesRef& Interceptor, const FName& MessageType )
+	FORCEINLINE void AddInterceptor( const IMessageInterceptorRef& Interceptor, const FName& MessageType )
 	{
 		EnqueueCommand(FSimpleDelegate::CreateRaw(this, &FMessageRouter::HandleAddInterceptor, Interceptor, MessageType));
 	}
@@ -69,7 +69,7 @@ public:
 	 * @param Interceptor The interceptor to remove.
 	 * @param MessageType The type of messages to stop intercepting.
 	 */
-	FORCEINLINE void RemoveInterceptor( const IInterceptMessagesRef& Interceptor, const FName& MessageType )
+	FORCEINLINE void RemoveInterceptor( const IMessageInterceptorRef& Interceptor, const FName& MessageType )
 	{
 		EnqueueCommand(FSimpleDelegate::CreateRaw(this, &FMessageRouter::HandleRemoveInterceptor, Interceptor, MessageType));
 	}
@@ -200,7 +200,7 @@ private:
 private:
 
 	/** Handles adding message interceptors. */
-	void HandleAddInterceptor( IInterceptMessagesRef Interceptor, FName MessageType );
+	void HandleAddInterceptor( IMessageInterceptorRef Interceptor, FName MessageType );
 
 	/** Handles adding message recipients. */
 	void HandleAddRecipient( FMessageAddress Address, IReceiveMessagesWeakPtr RecipientPtr );
@@ -209,7 +209,7 @@ private:
 	void HandleAddSubscriber( IMessageSubscriptionRef Subscription );
 
 	/** Handles the removal of message interceptors. */
-	void HandleRemoveInterceptor( IInterceptMessagesRef Interceptor, FName MessageType );
+	void HandleRemoveInterceptor( IMessageInterceptorRef Interceptor, FName MessageType );
 
 	/** Handles the removal of message recipients. */
 	void HandleRemoveRecipient( FMessageAddress Address );
@@ -223,7 +223,7 @@ private:
 private:
 
 	/** Maps message types to interceptors. */
-	TMap<FName, TArray<IInterceptMessagesPtr>> ActiveInterceptors;
+	TMap<FName, TArray<IMessageInterceptorPtr>> ActiveInterceptors;
 
 	/** Maps message addresses to recipients. */
 	TMap<FMessageAddress, IReceiveMessagesWeakPtr> ActiveRecipients;
diff --git a/Engine/Source/Runtime/Messaging/Private/Bus/MessageTracer.cpp b/Engine/Source/Runtime/Messaging/Private/Bus/MessageTracer.cpp
--- a/Engine/Source/Runtime/Messaging/Private/Bus/MessageTracer.cpp
+++ b/Engine/Source/Runtime/Messaging/Private/Bus/MessageTracer.cpp
@@ -171,7 +171,7 @@ void FMessageTracer::ProcessHandledMessage( IMessageContextRef Context, double T
 }
 
 
-void FMessageTracer::ProcessRemovedInterceptor( IInterceptMessagesRef Interceptor, FName MessageType, double TimeSeconds )
+void FMessageTracer::ProcessRemovedInterceptor( IMessageInterceptorRef Interceptor, FName MessageType, double TimeSeconds )
 {
 
 }
diff --git a/Engine/Source/Runtime/Messaging/Private/Bus/MessageTracer.h b/Engine/Source/Runtime/Messaging/Private/Bus/MessageTracer.h
--- a/Engine/Source/Runtime/Messaging/Private/Bus/MessageTracer.h
+++ b/Engine/Source/Runtime/Messaging/Private/Bus/MessageTracer.h
@@ -56,7 +56,7 @@ public:
 	 * @param Interceptor The added interceptor.
 	 * @param MessageType The type of messages being intercepted.
 	 */
-	void TraceAddedInterceptor( const IInterceptMessagesRef& Interceptor, const FName& MessageType )
+	void TraceAddedInterceptor( const IMessageInterceptorRef& Interceptor, const FName& MessageType )
 	{
 		FString Name = TEXT("@todo");
 		Traces.Enqueue(FSimpleDelegate::CreateRaw(this, &FMessageTracer::ProcessAddedInterceptor, Name, MessageType, FPlatformTime::Seconds()));
@@ -126,7 +126,7 @@ public:
 	 * @param Context The context of the intercepted message.
 	 * @param Interceptor The interceptor.
 	 */
-	void TraceInterceptedMessage( const IMessageContextRef& Context, const IInterceptMessagesRef& Interceptor )
+	void TraceInterceptedMessage( const IMessageContextRef& Context, const IMessageInterceptorRef& Interceptor )
 	{
 		if (Running)
 		{
@@ -140,7 +140,7 @@ public:
 	 * @param Interceptor The removed interceptor.
 	 * @param MessageType The type of messages that is no longer being intercepted.
 	 */
-	void TraceRemovedInterceptor( const IInterceptMessagesRef& Interceptor, const FName& MessageType )
+	void TraceRemovedInterceptor( const IMessageInterceptorRef& Interceptor, const FName& MessageType )
 	{
 		if (!Running)
 		{
@@ -374,7 +374,7 @@ protected:
 	 * @param MessageType The type of messages no longer being intercepted.
 	 * @param TimeSecond The time at which the interceptor was removed.
 	 */
-	void ProcessRemovedInterceptor( IInterceptMessagesRef Interceptor, FName MessageType, double TimeSeconds );
+	void ProcessRemovedInterceptor( IMessageInterceptorRef Interceptor, FName MessageType, double TimeSeconds );
 
 	/**
 	 * Processes traces for removed message recipients.
diff --git a/Engine/Source/Runtime/Messaging/Private/MessagingModule.cpp b/Engine/Source/Runtime/Messaging/Private/MessagingModule.cpp
--- a/Engine/Source/Runtime/Messaging/Private/MessagingModule.cpp
+++ b/Engine/Source/Runtime/Messaging/Private/MessagingModule.cpp
@@ -54,7 +54,7 @@ public:
 
 	// IMessagingModule interface
 
-	virtual IMessageBridgePtr CreateBridge( const FMessageAddress& Address, const IMessageBusRef& Bus, const ISerializeMessagesRef& Serializer, const ITransportMessagesRef& Transport ) override
+	virtual IMessageBridgePtr CreateBridge( const FMessageAddress& Address, const IMessageBusRef& Bus, const ISerializeMessagesRef& Serializer, const IMessageTransportRef& Transport ) override
 	{
 		return MakeShareable(new FMessageBridge(Address, Bus, Serializer, Transport));
 	}
diff --git a/Engine/Source/Runtime/Messaging/Public/Common/MessageBridgeBuilder.h b/Engine/Source/Runtime/Messaging/Public/Common/MessageBridgeBuilder.h
--- a/Engine/Source/Runtime/Messaging/Public/Common/MessageBridgeBuilder.h
+++ b/Engine/Source/Runtime/Messaging/Public/Common/MessageBridgeBuilder.h
@@ -63,7 +63,7 @@ public:
 	 * @param InTransport The transport technology to use.
 	 * @return This instance (for method chaining).
 	 */
-	FMessageBridgeBuilder& UsingTransport( const ITransportMessagesRef& InTransport )
+	FMessageBridgeBuilder& UsingTransport( const IMessageTransportRef& InTransport )
 	{
 		Transport = InTransport;
 
@@ -157,5 +157,5 @@ private:
 	ISerializeMessagesPtr Serializer;
 
 	/** Holds a reference to the message transport technology. */
-	ITransportMessagesPtr Transport;
+	IMessageTransportPtr Transport;
 };
diff --git a/Engine/Source/Runtime/Messaging/Public/Interfaces/IAuthorizeMessageRecipients.h b/Engine/Source/Runtime/Messaging/Public/Interfaces/IAuthorizeMessageRecipients.h
--- a/Engine/Source/Runtime/Messaging/Public/Interfaces/IAuthorizeMessageRecipients.h
+++ b/Engine/Source/Runtime/Messaging/Public/Interfaces/IAuthorizeMessageRecipients.h
@@ -4,7 +4,7 @@
 
 
 struct FMessageAddress;
-class IInterceptMessages;
+class IMessageInterceptor;
 class IReceiveMessages;
 
 
@@ -22,7 +22,7 @@ public:
 	 * @param MessageType The type of messages to intercept.
 	 * @return true if the request was authorized, false otherwise.
 	 */
-	virtual bool AuthorizeInterceptor( const TSharedRef<IInterceptMessages, ESPMode::ThreadSafe>& Interceptor, const FName& MessageType ) = 0;
+	virtual bool AuthorizeInterceptor( const TSharedRef<IMessageInterceptor, ESPMode::ThreadSafe>& Interceptor, const FName& MessageType ) = 0;
 
 	/**
 	 * Authorizes a request to register the specified recipient.
diff --git a/Engine/Source/Runtime/Messaging/Public/Interfaces/IInterceptMessages.h b/Engine/Source/Runtime/Messaging/Public/Interfaces/IInterceptMessages.h
deleted file mode 100644
diff --git a/Engine/Source/Runtime/Messaging/Public/Interfaces/IMessageBus.h b/Engine/Source/Runtime/Messaging/Public/Interfaces/IMessageBus.h
--- a/Engine/Source/Runtime/Messaging/Public/Interfaces/IMessageBus.h
+++ b/Engine/Source/Runtime/Messaging/Public/Interfaces/IMessageBus.h
@@ -5,7 +5,7 @@
 
 enum class EMessageScope;
 struct FMessageAddress;
-class IInterceptMessages;
+class IMessageInterceptor;
 class IMessageAttachment;
 class IMessageContext;
 class IMessageSubscription;
@@ -129,7 +129,7 @@ public:
 	 * @param MessageType The type of messages to intercept.
 	 * @see Unintercept
 	 */
-	virtual void Intercept( const TSharedRef<IInterceptMessages, ESPMode::ThreadSafe>& Interceptor, const FName& MessageType ) = 0;
+	virtual void Intercept( const TSharedRef<IMessageInterceptor, ESPMode::ThreadSafe>& Interceptor, const FName& MessageType ) = 0;
 
 	/**
 	 * Sends a message to subscribed recipients.
@@ -197,7 +197,7 @@ public:
 	 * @param MessageType The type of messages to stop intercepting.
 	 * @see Intercept
 	 */
-	virtual void Unintercept( const TSharedRef<IInterceptMessages, ESPMode::ThreadSafe>& Interceptor, const FName& MessageType ) = 0;
+	virtual void Unintercept( const TSharedRef<IMessageInterceptor, ESPMode::ThreadSafe>& Interceptor, const FName& MessageType ) = 0;
 
 	/**
 	 * Unregisters a message recipient from the message bus.
diff --git a/Engine/Source/Runtime/Messaging/Public/Interfaces/IMessageInterceptor.h b/Engine/Source/Runtime/Messaging/Public/Interfaces/IMessageInterceptor.h
new file mode 100644
--- /dev/null
+++ b/Engine/Source/Runtime/Messaging/Public/Interfaces/IMessageInterceptor.h
@@ -0,0 +1,35 @@
+// Copyright 1998-2014 Epic Games, Inc. All Rights Reserved.
+
+#pragma once
+
+
+class IMessageContext;
+
+
+/**
+ * Interface for message interceptors.
+ */
+class IMessageInterceptor
+{
+public:
+
+	/**
+	 * Intercepts a message before it is being passed to the message router.
+	 *
+	 * @param Context The context of the message to intercept.
+	 * @return true if the message was intercepted and should not be routed, false otherwise.
+	 */
+	virtual bool InterceptMessage( const TSharedRef<IMessageContext, ESPMode::ThreadSafe>& Context ) = 0;
+
+public:
+
+	/** Virtual destructor. */
+	virtual ~IMessageInterceptor() { }
+};
+
+
+/** Type definition for shared pointers to instances of IInterceptMessages. */
+typedef TSharedPtr<IMessageInterceptor, ESPMode::ThreadSafe> IMessageInterceptorPtr;
+
+/** Type definition for shared references to instances of IInterceptMessages. */
+typedef TSharedRef<IMessageInterceptor, ESPMode::ThreadSafe> IMessageInterceptorRef;
diff --git a/Engine/Source/Runtime/Messaging/Public/Interfaces/IMessageTransport.h b/Engine/Source/Runtime/Messaging/Public/Interfaces/IMessageTransport.h
new file mode 100644
--- /dev/null
+++ b/Engine/Source/Runtime/Messaging/Public/Interfaces/IMessageTransport.h
@@ -0,0 +1,113 @@
+// Copyright 1998-2014 Epic Games, Inc. All Rights Reserved.
+
+#pragma once
+
+
+class IMessageAttachment;
+class IMessageData;
+
+
+/**
+ * Delegate type for receiving transport messages.
+ *
+ * The first parameter is the serialized message content.
+ * The second parameter is the message attachment, if any.
+ * The third parameter is the identifier of the transport node that sent the message.
+ */
+DECLARE_DELEGATE_ThreeParams(FOnMessageTransportMessageReceived, FArchive&, const IMessageAttachmentPtr&, const FGuid&);
+
+/**
+ * Delegate type for handling the discovery of transport nodes.
+ *
+ * The first parameter is the identifier of the transport node that was discovered.
+ */
+DECLARE_DELEGATE_OneParam(FOnMessageTransportNodeDiscovered, const FGuid&);
+
+/**
+ * Delegate type for handling disconnected or timed out transport nodes.
+ *
+ * Remote nodes can be lost if they disconnect from the bus or if a keep-alive
+ * message was not received within the required timeout (i.e. in case of a crash).
+ *
+ * The first parameter is the identifier of the transport node that was lost.
+ */
+DECLARE_DELEGATE_OneParam(FOnMessageTransportNodeLost, const FGuid&);
+
+
+/**
+ * Interface for message transport technologies.
+ *
+ * Licensees can implement this interface to add support for custom message transport
+ * technologies that are not supported out of the box, i.e. custom network protocols or APIs.
+ */
+class IMessageTransport
+{
+public:
+
+	/**
+	 * Gets the name of this transport (for debugging purposes).
+	 *
+	 * @return The debug name.
+	 */
+	virtual FName GetDebugName() const = 0;
+
+	/**
+	 * Starts up the message transport.
+	 *
+	 * @return Whether the transport was started successfully.
+	 * @see StopTransport
+	 */
+	virtual bool StartTransport() = 0;
+
+	/**
+	 * Shuts down the message transport.
+	 *
+	 * @see StartTransport
+	 */
+	virtual void StopTransport() = 0;
+
+	/**
+	 * Transports the given message data to the specified network nodes.
+	 *
+	 * @param Data The serialized message data to transport.
+	 * @param Attachment An optional message attachment (i.e. file or memory buffer).
+	 * @param Recipients The transport nodes to send the message to.
+	 * @return true if the message is being transported, false otherwise.
+	 */
+	virtual bool TransportMessage( const TSharedRef<IMessageData, ESPMode::ThreadSafe>& Data, const TSharedPtr<IMessageAttachment, ESPMode::ThreadSafe>& Attachment, const TArray<FGuid>& Recipients ) = 0;
+
+public:
+
+	/**
+	 * Returns a delegate that is executed when message data has been received.
+	 *
+	 * @param Delegate The delegate to set.
+	 */
+	virtual FOnMessageTransportMessageReceived& OnMessageReceived() = 0;
+
+	/**
+	 * Returns a delegate that is executed when a transport node has been discovered.
+	 *
+	 * @param Delegate The delegate to set.
+	 */
+	virtual FOnMessageTransportNodeDiscovered& OnNodeDiscovered() = 0;
+
+	/**
+	 * Returns a delegate that is executed when a transport node has closed or timed out.
+	 *
+	 * @param Delegate The delegate to set.
+	 */
+	virtual FOnMessageTransportNodeLost& OnNodeLost() = 0;
+
+protected:
+
+	/** Virtual constructor. */
+	~IMessageTransport() { }
+};
+
+
+/** Type definition for shared pointers to instances of ITransportMessages. */
+typedef TSharedPtr<IMessageTransport, ESPMode::ThreadSafe> IMessageTransportPtr;
+
+/** Type definition for shared references to instances of ITransportMessages. */
+typedef TSharedRef<IMessageTransport, ESPMode::ThreadSafe> IMessageTransportRef;
diff --git a/Engine/Source/Runtime/Messaging/Public/Interfaces/IMessagingModule.h b/Engine/Source/Runtime/Messaging/Public/Interfaces/IMessagingModule.h
--- a/Engine/Source/Runtime/Messaging/Public/Interfaces/IMessagingModule.h
+++ b/Engine/Source/Runtime/Messaging/Public/Interfaces/IMessagingModule.h
@@ -10,7 +10,7 @@ class IAuthorizeMessageRecipients;
 class IMessageBus;
 class IMessageBridge;
 class ISerializeMessages;
-class ITransportMessages;
+class IMessageTransport;
 
 
 /**
@@ -36,7 +36,7 @@ public:
 	 * @return The new message bridge, or nullptr if the bridge couldn't be created.
 	 * @see CreateBus
 	 */
-	virtual TSharedPtr<IMessageBridge, ESPMode::ThreadSafe> CreateBridge( const FMessageAddress& Address, const TSharedRef<IMessageBus, ESPMode::ThreadSafe>& Bus, const TSharedRef<ISerializeMessages, ESPMode::ThreadSafe>& Serializer, const TSharedRef<ITransportMessages, ESPMode::ThreadSafe>& Transport ) = 0;
+	virtual TSharedPtr<IMessageBridge, ESPMode::ThreadSafe> CreateBridge( const FMessageAddress& Address, const TSharedRef<IMessageBus, ESPMode::ThreadSafe>& Bus, const TSharedRef<ISerializeMessages, ESPMode::ThreadSafe>& Serializer, const TSharedRef<IMessageTransport, ESPMode::ThreadSafe>& Transport ) = 0;
 
 	/**
 	 * Creates a new message bus.
diff --git a/Engine/Source/Runtime/Messaging/Public/Interfaces/ITransportMessages.h b/Engine/Source/Runtime/Messaging/Public/Interfaces/ITransportMessages.h
deleted file mode 100644
diff --git a/Engine/Source/Runtime/Messaging/Public/Messaging.h b/Engine/Source/Runtime/Messaging/Public/Messaging.h
--- a/Engine/Source/Runtime/Messaging/Public/Messaging.h
+++ b/Engine/Source/Runtime/Messaging/Public/Messaging.h
@@ -17,7 +17,7 @@
 #include "IMessageAttachment.h"
 #include "IMessageData.h"
 #include "IMessageContext.h"
-#include "IInterceptMessages.h"
+#include "IMessageInterceptor.h"
 #include "IMessageHandler.h"
 #include "IReceiveMessages.h"
 #include "ISendMessages.h"
@@ -28,7 +28,7 @@
 #include "IMessageBus.h"
 #include "IMutableMessageContext.h"
 #include "ISerializeMessages.h"
-#include "ITransportMessages.h"
+#include "IMessageTransport.h"
 #include "IMessageBridge.h"
 #include "IMessagingModule.h"
 
