Fixing Mac build

diff --git a/Engine/Source/Runtime/Messaging/Public/Interfaces/IReceiveMessages.h b/Engine/Source/Runtime/Messaging/Public/Interfaces/IReceiveMessages.h
--- a/Engine/Source/Runtime/Messaging/Public/Interfaces/IReceiveMessages.h
+++ b/Engine/Source/Runtime/Messaging/Public/Interfaces/IReceiveMessages.h
@@ -4,7 +4,6 @@
 
 
 // forward declarations
-enum ENamedThreads::Type;
 class IMessageContext;
 
 
