MessagingDebugger: The 'Clear' UI action now also clears the message counters in the endpoints list

diff --git a/Engine/Source/Runtime/Messaging/Private/Bus/MessageTracer.cpp b/Engine/Source/Runtime/Messaging/Private/Bus/MessageTracer.cpp
--- a/Engine/Source/Runtime/Messaging/Private/Bus/MessageTracer.cpp
+++ b/Engine/Source/Runtime/Messaging/Private/Bus/MessageTracer.cpp
@@ -258,6 +258,13 @@ void FMessageTracer::ResetMessages()
 	MessageInfos.Reset();
 	MessageTypes.Reset();
 
+	for (TMap<FMessageAddress, FMessageTracerEndpointInfoPtr>::TIterator It(AddressesToEndpointInfos); It; ++It)
+	{
+		FMessageTracerEndpointInfoPtr& EndpointInfo = It.Value();
+		EndpointInfo->ReceivedMessages.Reset();
+		EndpointInfo->SentMessages.Reset();
+	}
+
 	MessagesResetDelegate.Broadcast();
 }
 
