Messaging: Another pass towards "include what you use"

diff --git a/Engine/Source/Runtime/Messaging/Public/Common/FileMessageAttachment.h b/Engine/Source/Runtime/Messaging/Public/Common/FileMessageAttachment.h
deleted file mode 100644
diff --git a/Engine/Source/Runtime/Messaging/Public/Common/MessageBridgeBuilder.h b/Engine/Source/Runtime/Messaging/Public/Common/MessageBridgeBuilder.h
deleted file mode 100644
diff --git a/Engine/Source/Runtime/Messaging/Public/Common/MessageEndpoint.h b/Engine/Source/Runtime/Messaging/Public/Common/MessageEndpoint.h
deleted file mode 100644
diff --git a/Engine/Source/Runtime/Messaging/Public/Common/MessageEndpointBuilder.h b/Engine/Source/Runtime/Messaging/Public/Common/MessageEndpointBuilder.h
deleted file mode 100644
diff --git a/Engine/Source/Runtime/Messaging/Public/Deprecated/Messaging.h b/Engine/Source/Runtime/Messaging/Public/Deprecated/Messaging.h
new file mode 100644
--- /dev/null
+++ b/Engine/Source/Runtime/Messaging/Public/Deprecated/Messaging.h
@@ -0,0 +1,38 @@
+// Copyright 1998-2014 Epic Games, Inc. All Rights Reserved.
+
+#pragma once
+
+
+/* Public Dependencies
+ *****************************************************************************/
+
+#include "CoreUObject.h"
+
+
+/* Public Includes
+ *****************************************************************************/
+
+#include "IMessageAttachment.h"
+#include "IMessageContext.h"
+#include "IMessageInterceptor.h"
+#include "IMessageHandler.h"
+#include "IReceiveMessages.h"
+#include "ISendMessages.h"
+#include "IMessageSubscription.h"
+#include "IAuthorizeMessageRecipients.h"
+#include "IMessageTracerBreakpoint.h"
+#include "IMessageTracer.h"
+#include "IMessageBus.h"
+#include "IMutableMessageContext.h"
+#include "IMessageTransport.h"
+#include "IMessageBridge.h"
+#include "IMessagingModule.h"
+
+
+/* Common helpers
+ *****************************************************************************/
+
+#include "FileMessageAttachment.h"
+#include "MessageBridgeBuilder.h"
+#include "MessageEndpoint.h"
+#include "MessageEndpointBuilder.h"
diff --git a/Engine/Source/Runtime/Messaging/Public/Helpers/FileMessageAttachment.h b/Engine/Source/Runtime/Messaging/Public/Helpers/FileMessageAttachment.h
new file mode 100644
--- /dev/null
+++ b/Engine/Source/Runtime/Messaging/Public/Helpers/FileMessageAttachment.h
@@ -0,0 +1,65 @@
+// Copyright 1998-2014 Epic Games, Inc. All Rights Reserved.
+
+#pragma once
+
+#include "FileManager.h"
+#include "IMessageAttachment.h"
+
+
+/**
+ * Implements a message attachment whose data is held in a file.
+ *
+ * WARNING: Message attachments do not work yet for out of process messages.
+ */
+class FFileMessageAttachment
+	: public IMessageAttachment
+{
+public:
+
+	/**
+	 * Creates and initializes a new instance.
+	 *
+	 * @param Filename The full name and path of the file.
+	 */
+	FFileMessageAttachment( const FString& InFilename )
+		: AutoDeleteFile(false)
+		, Filename(InFilename)
+	{ }
+
+	/**
+	 * Creates and initializes a new instance.
+	 *
+	 * @param Filename The full name and path of the file.
+	 * @param AutoDeleteFile Whether to delete the file when this attachment is destroyed.
+	 */
+	FFileMessageAttachment( const FString& InFilename, bool InAutoDeleteFile )
+		: AutoDeleteFile(InAutoDeleteFile)
+		, Filename(InFilename)
+	{ }
+
+	/** Destructor. */
+	~FFileMessageAttachment()
+	{
+		if (AutoDeleteFile)
+		{
+			IFileManager::Get().Delete(*Filename);
+		}
+	}
+
+public:
+
+	// IMessageAttachment interface
+
+	virtual FArchive* CreateReader() override
+	{
+		return IFileManager::Get().CreateFileReader(*Filename);
+	}
+
+private:
+
+	/** Holds a flag indicating whether the file should be deleted. */
+	bool AutoDeleteFile;
+
+	/** Holds the name of the file that holds the attached data. */
+	FString Filename;
+};
diff --git a/Engine/Source/Runtime/Messaging/Public/Helpers/MessageBridgeBuilder.h b/Engine/Source/Runtime/Messaging/Public/Helpers/MessageBridgeBuilder.h
new file mode 100644
--- /dev/null
+++ b/Engine/Source/Runtime/Messaging/Public/Helpers/MessageBridgeBuilder.h
@@ -0,0 +1,139 @@
+// Copyright 1998-2014 Epic Games, Inc. All Rights Reserved.
+
+#pragma once
+
+#include "IMessageBridge.h"
+#include "IMessageBus.h"
+#include "IMessageContext.h"
+#include "IMessageTransport.h"
+#include "IMessagingModule.h"
+
+
+/**
+ * Implements a message bridge builder.
+ */
+class FMessageBridgeBuilder
+{
+public:
+
+	/** Default constructor. */
+	FMessageBridgeBuilder()
+		: Address(FMessageAddress::NewAddress())
+		, BusPtr(IMessagingModule::Get().GetDefaultBus())
+		, Disabled(false)
+		, Transport(nullptr)
+	{ }
+
+	/**
+	 * Creates and initializes a new instance.
+	 *
+	 * @param InBus The message bus to attach the bridge to.
+	 */
+	FMessageBridgeBuilder( const IMessageBusRef& Bus )
+		: Address(FMessageAddress::NewAddress())
+		, BusPtr(Bus)
+		, Disabled(false)
+		, Transport(nullptr)
+	{ }
+
+public:
+
+	/**
+	 * Disables the bridge.
+	 *
+	 * @return This instance (for method chaining).
+	 */
+	FMessageBridgeBuilder& ThatIsDisabled()
+	{
+		Disabled = true;
+
+		return *this;
+	}
+
+	/**
+	 * Configures the bridge to use a specific message transport technology.
+	 *
+	 * @param InTransport The transport technology to use.
+	 * @return This instance (for method chaining).
+	 */
+	FMessageBridgeBuilder& UsingTransport( const IMessageTransportRef& InTransport )
+	{
+		Transport = InTransport;
+
+		return *this;
+	}
+
+	/**
+	 * Sets the bridge's address.
+	 *
+	 * If no address is specified, one will be generated automatically.
+	 *
+	 * @param InAddress The address to set.
+	 * @return This instance (for method chaining).
+	 */
+	FMessageBridgeBuilder& WithAddress( const FMessageAddress& InAddress )
+	{
+		Address = InAddress;
+
+		return *this;
+	}
+
+public:
+
+	/**
+	 * Builds the message bridge as configured.
+	 *
+	 * @return A new message bridge, or nullptr if it couldn't be built.
+	 */
+	IMessageBridgePtr Build()
+	{
+		IMessageBridgePtr Bridge;
+
+		check(Transport.IsValid());
+
+		IMessageBusPtr Bus = BusPtr.Pin();
+
+		if (Bus.IsValid())
+		{
+			Bridge = IMessagingModule::Get().CreateBridge(Address, Bus.ToSharedRef(), Transport.ToSharedRef());
+
+			if (Bridge.IsValid())
+			{
+				if (Disabled)
+				{
+					Bridge->Disable();
+				}
+				else
+				{
+					Bridge->Enable();
+				}
+			}
+		}
+
+		return Bridge;
+	}
+
+	/**
+	 * Implicit conversion operator to build the message bridge as configured.
+	 *
+	 * @return A new message bridge, or nullptr if it couldn't be built.
+	 */
+	operator IMessageBridgePtr()
+	{
+		return Build();
+	}
+
+private:
+
+	/** Holds the bridge's address. */
+	FMessageAddress Address;
+
+	/** Holds a weak pointer to the message bus to attach to. */
+	IMessageBusWeakPtr BusPtr;
+
+	/** Holds a flag indicating whether the bridge should be disabled. */
+	bool Disabled;
+
+	/** Holds a reference to the message transport technology. */
+	IMessageTransportPtr Transport;
+};
diff --git a/Engine/Source/Runtime/Messaging/Public/Helpers/MessageEndpoint.h b/Engine/Source/Runtime/Messaging/Public/Helpers/MessageEndpoint.h
new file mode 100644
--- /dev/null
+++ b/Engine/Source/Runtime/Messaging/Public/Helpers/MessageEndpoint.h
@@ -0,0 +1,836 @@
+// Copyright 1998-2014 Epic Games, Inc. All Rights Reserved.
+
+#pragma once
+
+#include "IMessageAttachment.h"
+#include "IMessageBus.h"
+#include "IMessageContext.h"
+#include "IMessageHandler.h"
+#include "IReceiveMessages.h"
+#include "ISendMessages.h"
+
+
+/**
+ * DEPRECATED: Delegate type for error notifications.
+ *
+ * The first parameter is the context of the sent message.
+ * The second parameter is the error string.
+ */
+DECLARE_DELEGATE_TwoParams(FOnMessageEndpointError, const IMessageContextRef&, const FString&);
+
+/**
+ * Delegate type for received messages.
+ *
+ * The first parameter is the received message.
+ * The return value indicates whether the message should be handled.
+ */
+DECLARE_DELEGATE_RetVal_OneParam(bool, FOnMessageEndpointReceiveMessage, const IMessageContextRef&);
+
+
+/**
+ * Implements a message endpoint for sending and receiving messages on a message bus.
+ *
+ * This class provides a convenient implementation of the IRecieveMessages and ISendMessages interfaces,
+ * which allow consumers to send and receive messages on a message bus. The endpoint allows for receiving
+ * messages asynchronously as they arrive, as well as synchronously through an inbox that can be polled.
+ *
+ * By default, messages are received synchronously on the thread that the endpoint was created on.
+ * If the message consumer is thread-safe, a more efficient message dispatch can be enabled by calling
+ * the SetRecipientThread(ENamedThreads::AnyThread) method.
+ *
+ * Endpoints that are destroyed or receive messages on non-Game threads should use the static function
+ * FMessageEndpoint::SafeRelease() to dispose of the endpoint. This will ensure that there are no race
+ * conditions between endpoint destruction and the receiving of messages.
+ */
+class FMessageEndpoint
+	: public TSharedFromThis<FMessageEndpoint, ESPMode::ThreadSafe>
+	, public IReceiveMessages
+	, public ISendMessages
+{
+public:
+
+	/**
+	 * Type definition for the endpoint builder.
+	 *
+	 * When building message endpoints that receive messages on AnyThread, use the SafeRelease
+	 * helper function to avoid race conditions when destroying the objects that own the endpoints.
+	 *
+	 * @see SafeRelease
+	 */
+	typedef struct FMessageEndpointBuilder Builder;
+
+	/**
+	 * Creates and initializes a new instance.
+	 *
+	 * @param InName The endpoint's name (for debugging purposes).
+	 * @param InBus The message bus to attach this endpoint to.
+	 * @param InHandlers The collection of message handlers to register.
+	 */
+	FMessageEndpoint( const FName& InName, const IMessageBusRef& InBus, const TArray<IMessageHandlerPtr>& InHandlers )
+		: Address(FMessageAddress::NewAddress())
+		, BusPtr(InBus)
+		, Enabled(true)
+		, Handlers(InHandlers)
+		, Id(FGuid::NewGuid())
+		, InboxEnabled(false)
+		, Name(InName)
+	{ }
+
+	/** Destructor. */
+	~FMessageEndpoint()
+	{
+		IMessageBusPtr Bus = BusPtr.Pin();
+
+		if (Bus.IsValid())
+		{
+			Bus->Unregister(Address);
+		}
+	}
+
+public:
+
+	/**
+	 * Disables this endpoint.
+	 *
+	 * A disabled endpoint will not receive any subscribed messages until it is enabled again.
+	 * Endpoints should be created in an enabled state by default.
+	 *
+	 * @see Enable, IsEnabled
+	 */
+	void Disable()
+	{
+		Enabled = false;
+	}
+
+	/**
+	 * Enables this endpoint.
+	 *
+	 * An activated endpoint will receive subscribed messages.
+	 * Endpoints should be created in an enabled state by default.
+	 *
+	 * @see Disable, IsEnabled
+	 */
+	void Enable()
+	{
+		Enabled = true;
+	}
+
+	/**
+	 * Gets the endpoint's message address.
+	 *
+	 * @return Message address.
+	 */
+	const FMessageAddress& GetAddress() const
+	{
+		return Address;
+	}
+
+	/**
+	 * Checks whether this endpoint is connected to the bus.
+	 *
+	 * @return true if connected, false otherwise.
+	 * @see IsEnabled
+	 */
+	bool IsConnected() const
+	{
+		return BusPtr.IsValid();
+	}
+
+	/**
+	 * Checks whether this endpoint is enabled.
+	 *
+	 * @return true if the endpoint is enabled, false otherwise.
+	 * @see Disable, Enable, IsConnected
+	 */
+	bool IsEnabled() const
+	{
+		return Enabled;
+	}
+
+	/**
+	 * Sets the name of the thread to receive messages on.
+	 *
+	 * Use this method to receive messages on a particular thread, for example, if the
+	 * consumer owning this endpoint is not thread-safe. The default value is ThreadAny.
+	 *
+	 * ThreadAny is the fastest way to receive messages. It should be used if the receiving
+	 * code is completely thread-safe and if it is sufficiently fast. ThreadAny MUST NOT
+	 * be used if the receiving code is not thread-safe. It also SHOULD NOT be used if the
+	 * code includes time consuming operations, because it will block the message router,
+	 * causing no other messages to be delivered in the meantime.
+	 *
+	 * @param NamedThread The name of the thread to receive messages on.
+	 */
+	void SetRecipientThread( const ENamedThreads::Type& NamedThread )
+	{
+		RecipientThread = NamedThread;
+	}
+
+public:
+
+	/**
+	 * Defers processing of the given message by the specified time delay.
+	 *
+	 * The message is effectively delivered again to this endpoint after the
+	 * original sent time plus the time delay have elapsed.
+	 *
+	 * @param Context The context of the message to defer.
+	 * @param Delay The time delay.
+	 */
+	void Defer( const IMessageContextRef& Context, const FTimespan& Delay )
+	{
+		IMessageBusPtr Bus = GetBusIfEnabled();
+
+		if (Bus.IsValid())
+		{
+			Bus->Forward(Context, TArrayBuilder<FMessageAddress>().Add(Address), Delay, AsShared());
+		}
+	}
+
+	/**
+	 * Forwards a previously received message.
+	 *
+	 * Messages can only be forwarded to endpoints within the same process.
+	 *
+	 * @param Context The context of the message to forward.
+	 * @param Recipients The list of message recipients to forward the message to.
+	 * @param Delay The time delay.
+	 */
+	void Forward( const IMessageContextRef& Context, const TArray<FMessageAddress>& Recipients, const FTimespan& Delay )
+	{
+		IMessageBusPtr Bus = GetBusIfEnabled();
+
+		if (Bus.IsValid())
+		{
+			Bus->Forward(Context, Recipients, Delay, AsShared());
+		}
+	}
+
+	/**
+	 * Publishes a message to all subscribed recipients within the specified scope.
+	 *
+	 * @param Message The message to publish.
+	 * @param TypeInfo The message's type information.
+	 * @param Scope The message scope.
+	 * @param Fields The message content.
+	 * @param Delay The delay after which to publish the message.
+	 * @param Expiration The time at which the message expires.
+	 */
+	void Publish( void* Message, UScriptStruct* TypeInfo, EMessageScope Scope, const FTimespan& Delay, const FDateTime& Expiration )
+	{
+		IMessageBusPtr Bus = GetBusIfEnabled();
+
+		if (Bus.IsValid())
+		{
+			Bus->Publish(Message, TypeInfo, Scope, Delay, Expiration, AsShared());
+		}
+	}
+
+	/**
+	 * Sends a message to the specified list of recipients.
+	 *
+	 * @param Message The message to send.
+	 * @param TypeInfo The message's type information.
+	 * @param Attachment An optional binary data attachment.
+	 * @param Recipients The message recipients.
+	 * @param Delay The delay after which to send the message.
+	 * @param Expiration The time at which the message expires.
+	 */
+	void Send( void* Message, UScriptStruct* TypeInfo, const IMessageAttachmentPtr& Attachment, const TArray<FMessageAddress>& Recipients, const FTimespan& Delay, const FDateTime& Expiration )
+	{
+		IMessageBusPtr Bus = GetBusIfEnabled();
+
+		if (Bus.IsValid())
+		{
+			Bus->Send(Message, TypeInfo, Attachment, Recipients, Delay, Expiration, AsShared());
+		}
+	}
+
+	/**
+	 * Subscribes a message handler.
+	 *
+	 * @param MessageType The type name of the messages to subscribe to.
+	 * @param ScopeRange The range of message scopes to include in the subscription.
+	 */
+	void Subscribe( const FName& MessageType, const FMessageScopeRange& ScopeRange )
+	{
+		IMessageBusPtr Bus = GetBusIfEnabled();
+
+		if (Bus.IsValid())
+		{
+			Bus->Subscribe(AsShared(), MessageType, ScopeRange);
+		}
+	}
+
+	/**
+	 * Unsubscribes this endpoint from the specified message type.
+	 *
+	 * @param MessageType The type of message to unsubscribe (NAME_All = all types).
+	 * @see Subscribe
+	 */
+	void Unsubscribe( const FName& TopicPattern )
+	{
+		IMessageBusPtr Bus = GetBusIfEnabled();
+
+		if (Bus.IsValid())
+		{
+			Bus->Unsubscribe(AsShared(), TopicPattern);
+		}
+	}
+
+public:
+
+	/**
+	 * Disables the inbox for unhandled messages.
+	 *
+	 * The inbox is disabled by default.
+	 *
+	 * @see EnableInbox, IsInboxEmpty, IsInboxEnabled, ProcessInbox, ReceiveFromInbox
+	 */
+	void DisableInbox()
+	{
+		InboxEnabled = false;
+	}
+
+	/**
+	 * Enables the inbox for unhandled messages.
+	 *
+	 * If enabled, the inbox will queue up all received messages. Use ProcessInbox() to synchronously
+	 * invoke the registered message handlers for all queued up messages, or ReceiveFromInbox() to
+	 * manually receive one message from the inbox at a time. The inbox is disabled by default.
+	 *
+	 * @see DisableInbox, IsInboxEmpty, IsInboxEnabled, ProcessInbox, ReceiveFromInbox
+	 */
+	void EnableInbox()
+	{
+		InboxEnabled = true;
+	}
+
+	/**
+	 * Checks whether the inbox is empty.
+	 *
+	 * @return true if the inbox is empty, false otherwise.
+	 * @see DisableInbox, EnableInbox, IsInboxEnabled, ProcessInbox, ReceiveFromInbox
+	 */
+	bool IsInboxEmpty() const
+	{
+		return Inbox.IsEmpty();
+	}
+
+	/**
+	 * Checks whether the inbox is enabled.
+	 *
+	 * @see DisableInbox, EnableInbox, IsInboxEmpty, ProcessInbox, ReceiveFromInbox
+	 */
+	bool IsInboxEnabled() const
+	{
+		return InboxEnabled;
+	}
+
+	/**
+	 * Calls the matching message handlers for all messages queued up in the inbox.
+	 *
+	 * Note that an incoming message will only be queued up in the endpoint's inbox if the inbox has
+	 * been enabled and no matching message handler handled it. The inbox is disabled by default and
+	 * must be enabled using the EnableInbox() method.
+	 *
+	 * @see IsInboxEmpty, ReceiveFromInbox
+	 */
+	void ProcessInbox()
+	{
+		IMessageContextPtr Context;
+
+		while (Inbox.Dequeue(Context))
+		{
+			ProcessMessage(Context.ToSharedRef());
+		}
+	}
+
+	/**
+	 * Receives a single message from the endpoint's inbox.
+	 *
+	 * Note that an incoming message will only be queued up in the endpoint's inbox if the inbox has
+	 * been enabled and no matching message handler handled it. The inbox is disabled by default and
+	 * must be enabled using the EnableInbox() method.
+	 *
+	 * @param OutContext Will hold the context of the received message.
+	 * @return true if a message was received, false if the inbox was empty.
+	 * @see DisableInbox, EnableInbox, IsInboxEnabled, ProcessInbox
+	 */
+	bool ReceiveFromInbox( IMessageContextPtr& OutContext )
+	{
+		return Inbox.Dequeue(OutContext);
+	}
+
+public:
+
+	/**
+	 * Returns a delegate to be invoked when the endpoint receives a message.
+	 *
+	 * This callback delegate is called every time a message is received by this
+	 * endpoint and before it is being handled by any registered message handlers.
+	 *
+	 * @return The delegate.
+	 */
+	FOnMessageEndpointReceiveMessage& OnReceiveMessage()
+	{
+		return ReceiveDelegate;
+	}
+
+public:
+
+	// IReceiveMessages interface
+
+	virtual FName GetDebugName() const override
+	{
+		return Name;
+	}
+
+	virtual const FGuid& GetRecipientId() const override
+	{
+		return Id;
+	}
+
+	virtual ENamedThreads::Type GetRecipientThread() const override
+	{
+		return RecipientThread;
+	}
+
+	virtual bool IsLocal() const override
+	{
+		return true;
+	}
+
+	virtual void ReceiveMessage( const IMessageContextRef& Context ) override
+	{
+		if (!Enabled)
+		{
+			return;
+		}
+
+		if (ReceiveDelegate.IsBound() && !ReceiveDelegate.Execute(Context))
+		{
+			return;
+		}
+
+		if (InboxEnabled)
+		{
+			Inbox.Enqueue(Context);
+		}
+		else
+		{
+			ProcessMessage(Context);
+		}
+	}
+
+public:
+
+	// ISendMessages interface
+
+	virtual FMessageAddress GetSenderAddress() override
+	{
+		return Address;
+	}
+
+	virtual void NotifyMessageError( const IMessageContextRef& Context, const FString& Error ) override
+	{
+		ErrorDelegate.ExecuteIfBound(Context, Error);
+	}
+
+public:
+
+	/**
+	 * Immediately forwards a previously received message to the specified recipient.
+	 *
+	 * Messages can only be forwarded to endpoints within the same process.
+	 *
+	 * @param Context The context of the message to forward.
+	 * @param Recipient The address of the recipient to forward the message to.
+	 */
+	void Forward( const IMessageContextRef& Context, const FMessageAddress& Recipient )
+	{
+		Forward(Context, TArrayBuilder<FMessageAddress>().Add(Recipient), FTimespan::Zero());
+	}
+
+	/**
+	 * Forwards a previously received message to the specified recipient after a given delay.
+	 *
+	 * Messages can only be forwarded to endpoints within the same process.
+	 *
+	 * @param Context The context of the message to forward.
+	 * @param Recipient The address of the recipient to forward the message to.
+	 * @param ForwardingScope The scope of the forwarded message.
+	 * @param Delay The delay after which to publish the message.
+	 */
+	void Forward( const IMessageContextRef& Context, const FMessageAddress& Recipient, const FTimespan& Delay )
+	{
+		Forward(Context, TArrayBuilder<FMessageAddress>().Add(Recipient), Delay);
+	}
+
+	/**
+	 * Immediately forwards a previously received message to the specified list of recipients.
+	 *
+	 * Messages can only be forwarded to endpoints within the same process.
+	 *
+	 * @param Context The context of the message to forward.
+	 * @param Recipients The list of message recipients to forward the message to.
+	 * @param ForwardingScope The scope of the forwarded message.
+	 */
+	void Forward( const IMessageContextRef& Context, const TArray<FMessageAddress>& Recipients )
+	{
+		Forward(Context, Recipients, FTimespan::Zero());
+	}
+
+	/**
+	 * Immediately publishes a message to all subscribed recipients.
+	 *
+	 * @param Message The message to publish.
+	 */
+	template<typename MessageType>
+	void Publish( MessageType* Message )
+	{
+		Publish(Message, MessageType::StaticStruct(), EMessageScope::Network, FTimespan::Zero(), FDateTime::MaxValue());
+	}
+
+	/**
+	 * Immediately pa message to all subscribed recipients within the specified scope.
+	 *
+	 * @param Message The message to publish.
+	 * @param Scope The message scope.
+	 */
+	template<typename MessageType>
+	void Publish( MessageType* Message, EMessageScope Scope )
+	{
+		Publish(Message, MessageType::StaticStruct(), Scope, FTimespan::Zero(), FDateTime::MaxValue());
+	}
+
+	/**
+	 * Publishes a message to all subscribed recipients after a given delay.
+	 *
+	 * @param Message The message to publish.
+	 * @param Delay The delay after which to publish the message.
+	 */
+	template<typename MessageType>
+	void Publish( MessageType* Message, const FTimespan& Delay )
+	{
+		Publish(Message, MessageType::StaticStruct(), EMessageScope::Network, Delay, FDateTime::MaxValue());
+	}
+
+	/**
+	 * Publishes a message to all subscribed recipients within the specified scope after a given delay.
+	 *
+	 * @param Message The message to publish.
+	 * @param Scope The message scope.
+	 * @param Delay The delay after which to publish the message.
+	 */
+	template<typename MessageType>
+	void Publish( MessageType* Message, EMessageScope Scope, const FTimespan& Delay )
+	{
+		Publish(Message, MessageType::StaticStruct(), Scope, Delay, FDateTime::MaxValue());
+	}
+
+	/**
+	 * Publishes a message to all subscribed recipients within the specified scope.
+	 *
+	 * @param Message The message to publish.
+	 * @param Scope The message scope.
+	 * @param Fields The message content.
+	 * @param Delay The delay after which to publish the message.
+	 * @param Expiration The time at which the message expires.
+	 */
+	template<typename MessageType>
+	void Publish( MessageType* Message, EMessageScope Scope, const FTimespan& Delay, const FDateTime& Expiration )
+	{
+		Publish(Message, MessageType::StaticStruct(), Scope, Delay, Expiration);
+	}
+
+	/**
+	 * Immediately sends a message to the specified recipient.
+	 *
+	 * @param MessageType The type of message to send.
+	 * @param Message The message to send.
+	 * @param Recipient The message recipient.
+	 */
+	template<typename MessageType>
+	void Send( MessageType* Message, const FMessageAddress& Recipient )
+	{
+		Send(Message, MessageType::StaticStruct(), nullptr, TArrayBuilder<FMessageAddress>().Add(Recipient), FTimespan::Zero(), FDateTime::MaxValue());
+	}
+
+	/**
+	 * Sends a message to the specified recipient after a given delay.
+	 *
+	 * @param MessageType The type of message to send.
+	 * @param Message The message to send.
+	 * @param Recipient The message recipient.
+	 * @param Delay The delay after which to send the message.
+	 */
+	template<typename MessageType>
+	void Send( MessageType* Message, const FMessageAddress& Recipient, const FTimespan& Delay )
+	{
+		Send(Message, MessageType::StaticStruct(), nullptr, TArrayBuilder<FMessageAddress>().Add(Recipient), Delay, FDateTime::MaxValue());
+	}
+
+	/**
+	 * Sends a message with fields and expiration to the specified recipient after a given delay.
+	 *
+	 * @param MessageType The type of message to send.
+	 * @param Message The message to send.
+	 * @param Recipient The message recipient.
+	 * @param Expiration The time at which the message expires.
+	 * @param Delay The delay after which to send the message.
+	 */
+	template<typename MessageType>
+	void Send( MessageType* Message, const FMessageAddress& Recipient, const FTimespan& Delay, const FDateTime& Expiration )
+	{
+		Send(Message, MessageType::StaticStruct(), nullptr, TArrayBuilder<FMessageAddress>().Add(Recipient), Delay, Expiration);
+	}
+
+	/**
+	 * Sends a message with fields and attachment to the specified recipient.
+	 *
+	 * @param MessageType The type of message to send.
+	 * @param Message The message to send.
+	 * @param Attachment An optional binary data attachment.
+	 * @param Recipient The message recipient.
+	 */
+	template<typename MessageType>
+	void Send( MessageType* Message, const IMessageAttachmentPtr& Attachment, const FMessageAddress& Recipient )
+	{
+		Send(Message, MessageType::StaticStruct(), Attachment, TArrayBuilder<FMessageAddress>().Add(Recipient), FTimespan::Zero(), FDateTime::MaxValue());
+	}
+
+	/**
+	 * Sends a message with fields, attachment and expiration to the specified recipient after a given delay.
+	 *
+	 * @param MessageType The type of message to send.
+	 * @param Message The message to send.
+	 * @param Attachment An optional binary data attachment.
+	 * @param Recipient The message recipient.
+	 * @param Expiration The time at which the message expires.
+	 * @param Delay The delay after which to send the message.
+	 */
+	template<typename MessageType>
+	void Send( MessageType* Message, const IMessageAttachmentPtr& Attachment, const FMessageAddress& Recipient, const FDateTime& Expiration, const FTimespan& Delay )
+	{
+		Send(Message, MessageType::StaticStruct(), Attachment, TArrayBuilder<FMessageAddress>().Add(Recipient), Delay, Expiration);
+	}
+
+	/**
+	 * Immediately sends a message to the specified list of recipients.
+	 *
+	 * @param MessageType The type of message to send.
+	 * @param Message The message to send.
+	 * @param Recipients The message recipients.
+	 */
+	template<typename MessageType>
+	void Send( MessageType* Message, const TArray<FMessageAddress>& Recipients )
+	{
+		Send(Message, MessageType::StaticStruct(), nullptr, Recipients, FTimespan::Zero(), FDateTime::MaxValue());
+	}
+
+	/**
+	 * Sends a message to the specified list of recipients after a given delay after a given delay.
+	 *
+	 * @param MessageType The type of message to send.
+	 * @param Message The message to send.
+	 * @param Recipients The message recipients.
+	 * @param Delay The delay after which to send the message.
+	 */
+	template<typename MessageType>
+	void Send( MessageType* Message, const TArray<FMessageAddress>& Recipients, const FTimespan& Delay )
+	{
+		Send(Message, MessageType::StaticStruct(), nullptr, Recipients, Delay, FDateTime::MaxValue());
+	}
+
+	/**
+	 * Sends a message with fields and attachment to the specified list of recipients after a given delay.
+	 *
+	 * @param MessageType The type of message to send.
+	 * @param Message The message to send.
+	 * @param Attachment An optional binary data attachment.
+	 * @param Recipients The message recipients.
+	 * @param Delay The delay after which to send the message.
+	 */
+	template<typename MessageType>
+	void Send( MessageType* Message, const IMessageAttachmentPtr& Attachment, const TArray<FMessageAddress>& Recipients, const FTimespan& Delay )
+	{
+		Send(Message, MessageType::StaticStruct(), Attachment, Recipients, Delay, FDateTime::MaxValue());
+	}
+
+	/**
+	 * Sends a message to the specified list of recipients.
+	 *
+	 * @param MessageType The type of message to send.
+	 * @param Message The message to send.
+	 * @param Attachment An optional binary data attachment.
+	 * @param Recipients The message recipients.
+	 * @param Delay The delay after which to send the message.
+	 * @param Expiration The time at which the message expires.
+	 */
+	template<typename MessageType>
+	void Send( MessageType* Message, const IMessageAttachmentPtr& Attachment, const TArray<FMessageAddress>& Recipients, const FTimespan& Delay, const FDateTime& Expiration )
+	{
+		Send(Message, MessageType::StaticStruct(), Attachment, Recipients, Delay, Expiration);
+	}
+
+	/**
+	 * Template method to subscribe the message endpoint to the specified type of messages with the default message scope.
+	 *
+	 * The default message scope is all messages excluding loopback messages.
+	 *
+	 * @param HandlerType The type of the class handling the message.
+	 * @param MessageType The type of messages to subscribe to.
+	 * @param Handler The class handling the messages.
+	 * @param HandlerFunc The class function handling the messages.
+	 */
+	template<class MessageType>
+	void Subscribe()
+	{
+		Subscribe(MessageType::StaticStruct()->GetFName(), FMessageScopeRange::AtLeast(EMessageScope::Thread));
+	}
+
+	/**
+	 * Template method to subscribe the message endpoint to the specified type and scope of messages.
+	 *
+	 * @param HandlerType The type of the class handling the message.
+	 * @param MessageType The type of messages to subscribe to.
+	 * @param Handler The class handling the messages.
+	 * @param HandlerFunc The class function handling the messages.
+	 * @param ScopeRange The range of message scopes to include in the subscription.
+	 */
+	template<class MessageType>
+	void Subscribe( const FMessageScopeRange& ScopeRange )
+	{
+		Subscribe(MessageType::StaticStruct()->GetFName(), ScopeRange);
+	}
+
+	/**
+	 * Unsubscribes this endpoint from all message types.
+	 *
+	 * @see Subscribe
+	 */
+	void Unsubscribe()
+	{
+		Unsubscribe(NAME_All);
+	}
+
+	/**
+	 * Template method to unsubscribe the endpoint from the specified message type.
+	 *
+	 * @param MessageType The type of message to unsubscribe (NAME_All = all types).
+	 * @see Subscribe
+	 */
+	template<class MessageType>
+	void Unsubscribe()
+	{
+		Unsubscribe(MessageType::StaticStruct()->GetFName());
+	}
+
+public:
+
+	/**
+	 * Safely releases a message endpoint that is receiving messages on AnyThread.
+	 *
+	 * When an object that owns a message endpoint receiving on AnyThread is being destroyed,
+	 * it is possible that the endpoint can outlive the object for a brief period of time if
+	 * the Messaging system is dispatching messages to it. The purpose of this helper function
+	 * is to block the calling thread while any messages are being dispatched, so that the
+	 * endpoint does not invoke any message handlers after the object has been destroyed.
+	 *
+	 * Note: When calling this function make sure that no other object is holding on to
+	 * the endpoint, or otherwise the caller may get blocked forever.
+	 *
+	 * @param Endpoint The message endpoint to release.
+	 */
+	static void SafeRelease( TSharedPtr<FMessageEndpoint, ESPMode::ThreadSafe>& Endpoint )
+	{
+		TWeakPtr<FMessageEndpoint, ESPMode::ThreadSafe> EndpointPtr = Endpoint;
+		Endpoint.Reset();
+		while (EndpointPtr.IsValid());
+	}
+
+protected:
+
+	/**
+	 * Gets a shared pointer to the message bus if this endpoint is enabled.
+	 *
+	 * @return The message bus.
+	 */
+	FORCEINLINE IMessageBusPtr GetBusIfEnabled() const
+	{
+		if (Enabled)
+		{
+			return BusPtr.Pin();
+		}
+
+		return nullptr;
+	}
+
+	/**
+	 * Forwards the given message context to matching message handlers.
+	 *
+	 * @param Context The context of the message to handle.
+	 */
+	void ProcessMessage( const IMessageContextRef& Context )
+	{
+		if (!Context->IsValid())
+		{
+			return;
+		}
+
+		for (int32 HandlerIndex = 0; HandlerIndex < Handlers.Num(); ++HandlerIndex)
+		{
+			IMessageHandlerPtr& Handler = Handlers[HandlerIndex];
+
+			if (Handler->GetHandledMessageType() == Context->GetMessageType())
+			{
+				Handler->HandleMessage(Context);
+			}
+		}
+	}
+
+private:
+
+	/** Holds the endpoint's identifier. */
+	const FMessageAddress Address;
+
+	/** Holds a weak pointer to the message bus. */
+	IMessageBusWeakPtr BusPtr;
+
+	/** Hold a flag indicating whether this endpoint is active. */
+	bool Enabled;
+
+	/** Holds the registered message handlers. */
+	TArray<IMessageHandlerPtr> Handlers;
+
+	/** Holds the endpoint's unique identifier (for debugging purposes). */
+	const FGuid Id;
+
+	/** Holds the endpoint's message inbox. */
+	TQueue<IMessageContextPtr, EQueueMode::Mpsc> Inbox;
+
+	/** Holds a flag indicating whether the inbox is enabled. */
+	bool InboxEnabled;
+
+	/** Holds the endpoint's name (for debugging purposes). */
+	const FName Name;
+
+	/** Holds a delegate that is invoked when a message has been received. */
+	FOnMessageEndpointReceiveMessage ReceiveDelegate;
+
+	/** Holds the name of the thread on which to receive messages. */
+	ENamedThreads::Type RecipientThread;
+
+private:
+
+	/** Holds a delegate that is invoked in case of messaging errors. */
+	FOnMessageEndpointError ErrorDelegate;
+};
+
+
+/** Type definition for shared pointers to instances of FMessageEndpoint. */
+typedef TSharedPtr<FMessageEndpoint, ESPMode::ThreadSafe> FMessageEndpointPtr;
+
+/** Type definition for shared references to instances of FMessageEndpoint. */
+typedef TSharedRef<FMessageEndpoint, ESPMode::ThreadSafe> FMessageEndpointRef;
diff --git a/Engine/Source/Runtime/Messaging/Public/Helpers/MessageEndpointBuilder.h b/Engine/Source/Runtime/Messaging/Public/Helpers/MessageEndpointBuilder.h
new file mode 100644
--- /dev/null
+++ b/Engine/Source/Runtime/Messaging/Public/Helpers/MessageEndpointBuilder.h
@@ -0,0 +1,218 @@
+// Copyright 1998-2014 Epic Games, Inc. All Rights Reserved.
+
+#pragma once
+
+#include "IMessageBus.h"
+#include "IMessageHandler.h"
+#include "IMessagingModule.h"
+#include "MessageEndpoint.h"
+#include "TaskGraphInterfaces.h"
+
+
+/**
+ * Implements a message endpoint builder.
+ */
+struct FMessageEndpointBuilder
+{
+public:
+
+	/**
+	 * Creates and initializes a new builder using the default message bus.
+	 *
+	 * WARNING: This constructor must be called from the Game thread.
+	 *
+	 * @param InName The endpoint's name (for debugging purposes).
+	 * @param InBus The message bus to attach the endpoint to.
+	 */
+	FMessageEndpointBuilder( const FName& InName )
+		: BusPtr(IMessagingModule::Get().GetDefaultBus())
+		, Disabled(false)
+		, InboxEnabled(false)
+		, Name(InName)
+		, RecipientThread(FTaskGraphInterface::Get().GetCurrentThreadIfKnown())
+	{ }
+
+	/**
+	 * Creates and initializes a new builder using the specified message bus.
+	 *
+	 * @param InName The endpoint's name (for debugging purposes).
+	 * @param InBus The message bus to attach the endpoint to.
+	 */
+	FMessageEndpointBuilder( const FName& InName, const IMessageBusRef& InBus )
+		: BusPtr(InBus)
+		, Disabled(false)
+		, InboxEnabled(false)
+		, Name(InName)
+		, RecipientThread(FTaskGraphInterface::Get().GetCurrentThreadIfKnown())
+	{ }
+
+public:
+
+	/**
+	 * Adds a message handler for the given type of messages.
+	 *
+	 * @param HandlerType The type of the object handling the messages.
+	 * @param MessageType The type of messages to handle.
+	 * @param Handler The class handling the messages.
+	 * @param HandlerFunc The class function handling the messages.
+	 * @param ReceivingThread The thread on which to handle the message.
+	 * @return This instance (for method chaining).
+	 */
+	template<typename MessageType, typename HandlerType>
+	FMessageEndpointBuilder& Handling( HandlerType* Handler, typename TMessageHandlerFunc<MessageType, HandlerType>::Type HandlerFunc )
+	{
+		// @todo gmp: implement proper async message deserialization, so this can be removed
+		static_assert(TStructOpsTypeTraits<MessageType>::WithMessageHandling == true, "Please add a WithMessageHandling type trait.");
+
+		Handlers.Add(MakeShareable(new TMessageHandler<MessageType, HandlerType>(Handler, HandlerFunc)));
+
+		return *this;
+	}
+
+	/**
+	 * Configures the endpoint to receive messages on any thread.
+	 *
+	 * By default, the builder initializes the message endpoint to receive on the
+	 * current thread. Use this method to receive on any available thread instead.
+	 *
+	 * ThreadAny is the fastest way to receive messages. It should be used if the receiving
+	 * code is completely thread-safe and if it is sufficiently fast. ThreadAny MUST NOT
+	 * be used if the receiving code is not thread-safe. It also SHOULD NOT be used if the
+	 * code includes time consuming operations, because it will block the message router,
+	 * causing no other messages to be delivered in the meantime.
+	 *
+	 * @return This instance (for method chaining).
+	 * @see ReceivingOnThread
+	 */
+	FMessageEndpointBuilder& ReceivingOnAnyThread()
+	{
+		RecipientThread = ENamedThreads::AnyThread;
+
+		return *this;
+	}
+
+	/**
+	 * Configured the endpoint to receive messages on a specific thread.
+	 *
+	 * By default, the builder initializes the message endpoint to receive on the
+	 * current thread. Use this method to receive on a different thread instead.
+	 *
+	 * Also see the additional notes for ReceivingOnAnyThread().
+	 *
+	 * @param NamedThread The name of the thread to receive messages on.
+	 * @return This instance (for method chaining).
+	 * @see ReceivingOnAnyThread
+	 */
+	FMessageEndpointBuilder& ReceivingOnThread( ENamedThreads::Type NamedThread )
+	{
+		RecipientThread = NamedThread;
+
+		return *this;
+	}
+
+	/**
+	 * Disables the endpoint.
+	 *
+	 * @return This instance (for method chaining).
+	 */
+	FMessageEndpointBuilder& ThatIsDisabled()
+	{
+		Disabled = true;
+
+		return *this;
+	}
+
+	/**
+	 * Registers a message handler with the endpoint.
+	 *
+	 * @param Handler The handler to add.
+	 * @return This instance (for method chaining).
+	 */
+	FMessageEndpointBuilder& WithHandler( const IMessageHandlerRef& Handler )
+	{
+		Handlers.Add(Handler);
+
+		return *this;
+	}
+
+	/**
+	 * Enables the endpoint's message inbox.
+	 *
+	 * The inbox is disabled by default.
+	 *
+	 * @return This instance (for method chaining).
+	 */
+	FMessageEndpointBuilder& WithInbox()
+	{
+		InboxEnabled = true;
+
+		return *this;
+	}
+
+public:
+
+	/**
+	 * Builds the message endpoint as configured.
+	 *
+	 * @return A new message endpoint, or nullptr if it couldn't be built.
+	 */
+	FMessageEndpointPtr Build()
+	{
+		FMessageEndpointPtr Endpoint;
+
+		IMessageBusPtr Bus = BusPtr.Pin();
+		
+		if (Bus.IsValid())
+		{
+			Endpoint = MakeShareable(new FMessageEndpoint(Name, Bus.ToSharedRef(), Handlers));
+			Bus->Register(Endpoint->GetAddress(), Endpoint.ToSharedRef());
+
+			if (Disabled)
+			{
+				Endpoint->Disable();
+			}
+
+			if (InboxEnabled)
+			{
+				Endpoint->EnableInbox();
+				Endpoint->SetRecipientThread(ENamedThreads::AnyThread);
+			}
+			else
+			{
+				Endpoint->SetRecipientThread(RecipientThread);
+			}
+		}
+
+		return Endpoint;
+	}
+
+	/**
+	 * Implicit conversion operator to build the message endpoint as configured.
+	 *
+	 * @return The message endpoint.
+	 */
+	operator FMessageEndpointPtr()
+	{
+		return Build();
+	}
+	
+private:
+
+	/** Holds a reference to the message bus to attach to. */
+	IMessageBusWeakPtr BusPtr;
+
+	/** Holds a flag indicating whether the endpoint should be disabled. */
+	bool Disabled;
+
+	/** Holds the collection of message handlers to register. */
+	TArray<IMessageHandlerPtr> Handlers;
+
+	/** Holds a flag indicating whether the inbox should be enabled. */
+	bool InboxEnabled;
+
+	/** Holds the endpoint's name (for debugging purposes). */
+	FName Name;
+
+	/** Holds the name of the thread on which to receive messages. */
+	ENamedThreads::Type RecipientThread;
+};
diff --git a/Engine/Source/Runtime/Messaging/Public/IAuthorizeMessageRecipients.h b/Engine/Source/Runtime/Messaging/Public/IAuthorizeMessageRecipients.h
new file mode 100644
--- /dev/null
+++ b/Engine/Source/Runtime/Messaging/Public/IAuthorizeMessageRecipients.h
@@ -0,0 +1,74 @@
+// Copyright 1998-2014 Epic Games, Inc. All Rights Reserved.
+
+#pragma once
+
+
+// forward declarations
+struct FMessageAddress;
+class IMessageInterceptor;
+class IReceiveMessages;
+
+
+/**
+ * Interface for classes that authorize message subscriptions.
+ */
+class IAuthorizeMessageRecipients
+{
+public:
+
+	/**
+	 * Authorizes a request to intercept messages of the specified type.
+	 *
+	 * @param Interceptor The message interceptor to authorize.
+	 * @param MessageType The type of messages to intercept.
+	 * @return true if the request was authorized, false otherwise.
+	 */
+	virtual bool AuthorizeInterceptor( const TSharedRef<IMessageInterceptor, ESPMode::ThreadSafe>& Interceptor, const FName& MessageType ) = 0;
+
+	/**
+	 * Authorizes a request to register the specified recipient.
+	 *
+	 * @param Recipient The recipient to register.
+	 * @param Address The recipient's address.
+	 * @return true if the request was authorized, false otherwise.
+	 */
+	virtual bool AuthorizeRegistration( const TSharedRef<IReceiveMessages, ESPMode::ThreadSafe>& Recipient, const FMessageAddress& Address ) = 0;
+
+	/**
+	 * Authorizes a request to add a subscription for the specified topic pattern.
+	 *
+	 * @param Subscriber The subscriber.
+	 * @param TopicPattern The message topic pattern to subscribe to.
+	 * @return true if the request is authorized, false otherwise.
+	 */
+	virtual bool AuthorizeSubscription( const TSharedRef<IReceiveMessages, ESPMode::ThreadSafe>& Subscriber, const FName& TopicPattern ) = 0;
+
+	/**
+	 * Authorizes a request to unregister the specified recipient.
+	 *
+	 * @param Address The address of the recipient to unregister.
+	 * @return true if the request was authorized, false otherwise.
+	 */
+	virtual bool AuthorizeUnregistration( const FMessageAddress& Address ) = 0;
+
+	/**
+	 * Authorizes a request to remove a subscription for the specified topic pattern.
+	 *
+	 * @param Subscriber The subscriber.
+	 * @param TopicPattern The message topic pattern to unsubscribe from.
+	 * @return true if the request is authorized, false otherwise.
+	 */
+	virtual bool AuthorizeUnsubscription( const TSharedRef<IReceiveMessages, ESPMode::ThreadSafe>& Subscriber, const FName& TopicPattern ) = 0;
+
+public:
+
+	/** Virtual destructor. */
+	virtual ~IAuthorizeMessageRecipients() { }
+};
+
+
+/** Type definition for shared pointers to instances of IAuthorizeMessageRecipients. */
+typedef TSharedPtr<IAuthorizeMessageRecipients> IAuthorizeMessageRecipientsPtr;
+
+/** Type definition for shared references to instances of IAuthorizeMessageRecipients. */
+typedef TSharedRef<IAuthorizeMessageRecipients> IAuthorizeMessageRecipientsRef;
diff --git a/Engine/Source/Runtime/Messaging/Public/IMessageAttachment.h b/Engine/Source/Runtime/Messaging/Public/IMessageAttachment.h
new file mode 100644
--- /dev/null
+++ b/Engine/Source/Runtime/Messaging/Public/IMessageAttachment.h
@@ -0,0 +1,45 @@
+// Copyright 1998-2014 Epic Games, Inc. All Rights Reserved.
+
+#pragma once
+
+
+/**
+ * Interface for message attachments.
+ *
+ * Message attachments are optional bulk data objects that can be attached to a message, such as
+ * files or memory buffers. Transferring attachments within the same process is very efficient and
+ * amounts to copying a pointer to the data. On the other hand, sending message attachments to
+ * applications running in other processes or on other computers may - depending on their size -
+ * take a considerable amount of time, because they need to be serialized and transferred. For this
+ * reason, attachments are transferred separately from messages in such cases.
+ *
+ * NOTE: Message attachments are not fully implemented yet. In particular, transferring attachments
+ * over a network does not work yet. This API may change in the near future, so please use with care.
+ *
+ * @see IMessageContext
+ */
+class IMessageAttachment
+{
+public:
+
+	/**
+	 * Creates an archive reader to the data.
+	 *
+	 * The caller is responsible for deleting the returned object.
+	 *
+	 * @return An archive reader.
+	 */
+	virtual FArchive* CreateReader() = 0;
+
+public:
+
+	/** Virtual destructor. */
+	virtual ~IMessageAttachment() { }
+};
+
+
+/** Type definition for shared pointers to instances of IMessageAttachment. */
+typedef TSharedPtr<IMessageAttachment, ESPMode::ThreadSafe> IMessageAttachmentPtr;
+
+/** Type definition for shared references to instances of IMessageAttachment. */
+typedef TSharedRef<IMessageAttachment, ESPMode::ThreadSafe> IMessageAttachmentRef;
diff --git a/Engine/Source/Runtime/Messaging/Public/IMessageBridge.h b/Engine/Source/Runtime/Messaging/Public/IMessageBridge.h
new file mode 100644
--- /dev/null
+++ b/Engine/Source/Runtime/Messaging/Public/IMessageBridge.h
@@ -0,0 +1,63 @@
+// Copyright 1998-2014 Epic Games, Inc. All Rights Reserved.
+
+#pragma once
+
+
+/**
+ * Interface for message bridges.
+ *
+ * A message bridge connects a message bus with another messaging system. It is really just a regular
+ * message endpoint connected to a message bus that translates sent and received messages between the bus
+ * and some underlying transport technology. The transport technology is usually implemented in the form of a 
+ * Message Transport Plug-in, such as the UdpMessaging plug-in that ships with Unreal Engine.
+ *
+ * The most common use case for message bridges is to connect two Unreal Engine message buses running in
+ * separate processes or en different computers. Another common use case is to connect an Unreal Engine
+ * message bus to an entirely different messaging system that is not based on Unreal Engine.
+ *
+ * Message bridge instances can be created with the @see IMessagingModule.CreateBridge method.
+ */
+class IMessageBridge
+{
+public:
+
+	/**
+	 * Disables this bridge.
+	 *
+	 * A disabled bridge will not receive any subscribed messages until it is enabled again.
+	 * Bridges should be created in a disabled state by default and explicitly enabled.
+	 *
+	 * @see Enable, IsEnabled
+	 */
+	virtual void Disable() = 0;
+
+	/**
+	 * Enables this bridge.
+	 *
+	 * An activated bridge will receive subscribed messages.
+	 * Bridges should be created in a disabled state by default and explicitly enabled.
+	 *
+	 * @see Disable, IsEnabled
+	 */
+	virtual void Enable() = 0;
+
+	/**
+	 * Checks whether the bridge is currently enabled.
+	 *
+	 * @return true if the bridge is enabled, false otherwise.
+	 * @see Disable, Enable
+	 */
+	virtual bool IsEnabled() const = 0;
+
+public:
+
+	/** Virtual destructor. */
+	virtual ~IMessageBridge() { }
+};
+
+
+/** Type definition for shared pointers to instances of IMessageBridge. */
+typedef TSharedPtr<IMessageBridge, ESPMode::ThreadSafe> IMessageBridgePtr;
+
+/** Type definition for shared references to instances of IMessageBridge. */
+typedef TSharedRef<IMessageBridge, ESPMode::ThreadSafe> IMessageBridgeRef;
diff --git a/Engine/Source/Runtime/Messaging/Public/IMessageBus.h b/Engine/Source/Runtime/Messaging/Public/IMessageBus.h
new file mode 100644
--- /dev/null
+++ b/Engine/Source/Runtime/Messaging/Public/IMessageBus.h
@@ -0,0 +1,244 @@
+// Copyright 1998-2014 Epic Games, Inc. All Rights Reserved.
+
+#pragma once
+
+
+// forward declarations
+enum class EMessageScope;
+struct FMessageAddress;
+class IMessageAttachment;
+class IMessageContext;
+class IMessageInterceptor;
+class IMessageSubscription;
+class IMessageTracer;
+class IReceiveMessages;
+class ISendMessages;
+class UScriptStruct;
+
+
+/** Delegate type for message bus shutdowns. */
+DECLARE_MULTICAST_DELEGATE(FOnMessageBusShutdown);
+
+
+/**
+ * Interface for message buses.
+ *
+ * A message bus is the main logical component to facilitate communication between (possibly distributed) parts of an application
+ * using Message Passing as its underlying architectural pattern. It allows registered sender and recipient objects, subsequently
+ * referred to as Message Endpoints, to exchange structured data in the form of user defined messages[1].
+ *
+ * Depending on their usage, messages are classified into a number of types, such as commands, events and documents. In Unreal Engine,
+ * all these messages are modeled with regular built-in or user defined UStructs that may either be empty or contain data in the
+ * form of UProperty fields[2]. Before being dispatched, messages are internally wrapped into so called Message Context objects
+ * (@see IMessageContext) that contain additional information about a message, such as when it was sent and the sender and recipients.
+ *
+ * The sending and receiving of messages is not limited to message endpoints within the same thread or process, but may be extended
+ * to other applications running on the same computer or on other computers connected to a network with so called Message Transport
+ * Plug-ins, such as the UdpMessaging plug-in that ships with Unreal Engine. The main goal of a message bus is to hide the technical
+ * details of the underlying transport mechanisms, so that users can focus on implementing their distributed applications rather than
+ * worrying about how data gets from one endpoint to another. A message bus makes it look as if all senders and recipients are located
+ * within the same process, regardless of whether that is actually the case or not. It is possible to restrict the reach of messages
+ * using so called Message Scopes (@see EMessageScope).
+ *
+ * All message recipients (objects implementing the @see IReceiveMessages interface) must be registered with the message bus through the
+ * @see IMessageBus.Register method. Before a recipient is destroyed it should unregister itself using the @see IMessageBus.Unregister
+ * method. Message senders (objects implementing the @see ISendMessages interface) do not register with the bus, but instead pass a
+ * reference to themselves each time they send a message. The IReceiveMessages and ISendMessages both are very low-level interfaces
+ * into the messaging system. Most users will prefer to use instances of the @see FMessageEndpoint class instead, which provides a much
+ * more convenient way of sending and receiving messages.
+ *
+ * Message buses in Unreal Engine support the following two common messaging patterns: Request-Reply and Publish-Subscribe. Please note
+ * that higher level concepts, such as Remote Procedure Call (RPC) are not part of the messaging system, but we may provide them as
+ * separate features in the future.
+ *
+ * In the Request-Reply pattern a message is sent to one or more particular message recipients using the IMessageBus.Send method.
+ * Message recipients implement the @see IMessageRecipient interface and are uniquely identified by their addresses (@see FMessageAddress).
+ * After a message is received, the recipients may then reply with another message using the same IMessageBus.Send method. Alternatively,
+ * a previously received message may be forwarded to another recipient using the @see IMessageBus.Forward method. This pattern is useful
+ * when message recipients already know about each other and wish to communicate directly, i.e. to exchange commands or events.
+ *
+ * In the Publish-Subscribe pattern a message is sent to all message recipients on the bus using the @see IMessageBus.Publish method.
+ * Only recipients that previously subscribed to the type of the sent message using the @IMessageBus.Subscribe method will actually receive
+ * the message. All other recipients will not receive the message. After a published message is received, the recipients may respond with
+ * another message, either directly to the message sender using IMessageBus.Send method, or by publishing another message using the
+ * the IMessageBus.Publish method. This pattern is useful for discovering recipients on the message bus and to recipients that are unknown.
+ *
+ * Most applications will often use a combination of both Request-Reply and Publish-Subscribe. A typical implementation of distributed
+ * applications involves service providers and service consumers (a service can be any useful functionality provided by a system). The
+ * service consumers will often publish a special message in order to discover service providers. Service providers are subscribed to these
+ * special messages and will respond with another message that contains information about the provider and that is sent back directly to the
+ * consumer. The consumer is now aware of the service provider's existence and its message address, and it can then request services by
+ * directly sending all subsequent messages directly to it.
+ *
+ * The dispatching of sent and published messages to the correct recipients is accomplished by the Message Router, which is an internal
+ * component of the message bus. The message router maintains an address book of known message endpoints that allows it to determine
+ * the destination of messages. If a message cannot be delivered, it is forwarded to a so called Dead Letter Channel, which is a queue
+ * of messages that can be inspected for debugging purposes[3].
+ *
+ * It is possible to intercept messages before they are being routed to recipients by registering so called Message Interceptors (objects
+ * implementing the IInterceptMessage interface) with the message bus. This feature enables advanced use cases that require inspection or
+ * manipulation of messages contents, such as message filtering and enriching, splitting and aggregating, resequencing or authentication[4].
+ *
+ * The messaging system also provides a facility for debugging the system itself through the so called Message Tracer (an internal object
+ * implementing the @see IMessageTracer interface) that can be accessed with the IMessageBus.GetTracer() method. The message tracer is
+ * currently used in the Messaging Debugger - a visual debugging tool for the Messaging System in Unreal Frontend and the Unreal Editor.
+ *
+ * Notes:
+ *
+ * [1] The messaging implementation in Unreal Engine is more or less closely following the terminology in the Enterprise Integration
+ * Patterns book at http://www.eaipatterns.com/eaipatterns.html. A good introduction to messaging concepts can also be found on MSDN
+ * at http://msdn.microsoft.com/en-us/library/ff647328.aspx
+ *
+ * [2] It is currently also required for such structures to implement a specific TStructOpsTypeTraits type trait or otherwise they
+ * will not be recognized as legal message types. Search for 'WithMessagingHandling' in the code base to find existing examples. We
+ * hope to remove the need for these type traits soon.
+ *
+ * [3] Dead letter channels are not implemented yet.
+ *
+ * [4] There are currently no built-in mechanisms for authentication and authorization, but they are on the to-do list.
+ *
+ * @see IReceiveMessages, ISendMessages, FMessageAddress, FMessageEndpoint
+ */
+class IMessageBus
+{
+public:
+
+	/**
+	 * Forwards a previously received message.
+	 *
+	 * Messages can only be forwarded to endpoints within the same process.
+	 *
+	 * @param Context The context of the message to forward.
+	 * @param Recipients The list of message recipients to forward the message to.
+	 * @param Delay The deferral time.
+	 * @param Forwarder The sender that forwards the message.
+	 * @see Publish, Send
+	 */
+	virtual void Forward( const TSharedRef<IMessageContext, ESPMode::ThreadSafe>& Context, const TArray<FMessageAddress>& Recipients, const FTimespan& Delay, const TSharedRef<ISendMessages, ESPMode::ThreadSafe>& Forwarder ) = 0;
+
+	/**
+	 * Gets the message bus tracer.
+	 *
+	 * @return Weak pointer to the tracer object.
+	 */
+	virtual TSharedRef<IMessageTracer, ESPMode::ThreadSafe> GetTracer() = 0;
+
+	/**
+	 * Adds an interceptor for messages of the specified type.
+	 *
+	 * @param Interceptor The interceptor.
+	 * @param MessageType The type of messages to intercept.
+	 * @see Unintercept
+	 */
+	virtual void Intercept( const TSharedRef<IMessageInterceptor, ESPMode::ThreadSafe>& Interceptor, const FName& MessageType ) = 0;
+
+	/**
+	 * Sends a message to subscribed recipients.
+	 *
+	 * The message is published with a default time to live of Process.
+	 *
+	 * @param Message The message to publish.
+	 * @param TypeInfo The message's type information.
+	 * @param Scope The message scope.
+	 * @param Delay The delay after which to send the message.
+	 * @param Expiration The time at which the message expires.
+	 * @param Publisher The message publisher.
+	 * @see Forward, Send
+	 */
+	virtual void Publish( void* Message, UScriptStruct* TypeInfo, EMessageScope Scope, const FTimespan& Delay, const FDateTime& Expiration, const TSharedRef<ISendMessages, ESPMode::ThreadSafe>& Publisher ) = 0;
+
+	/**
+	 * Registers a message recipient with the message bus.
+	 *
+	 * @param Address The address of the recipient to register.
+	 * @param Recipient The message recipient.
+	 * @see Unregister
+	 */
+	virtual void Register( const FMessageAddress& Address, const TSharedRef<IReceiveMessages, ESPMode::ThreadSafe>& Recipient) = 0;
+
+	/**
+	 * Sends a message to multiple recipients.
+	 *
+	 * @param Message The message to send.
+	 * @param TypeInfo The message's type information.
+	 * @param Attachment The binary data to attach to the message.
+	 * @param Recipients The list of message recipients.
+	 * @param Delay The delay after which to send the message.
+	 * @param Expiration The time at which the message expires.
+	 * @param Sender The message sender.
+	 * @see Forward, Publish
+	 */
+	virtual void Send( void* Message, UScriptStruct* TypeInfo, const TSharedPtr<IMessageAttachment, ESPMode::ThreadSafe>& Attachment, const TArray<FMessageAddress>& Recipients, const FTimespan& Delay, const FDateTime& Expiration, const TSharedRef<ISendMessages, ESPMode::ThreadSafe>& Sender ) = 0;
+
+	/**
+	 * Shuts down the message bus.
+	 *
+	 * @see OnShutdown
+	 */
+	virtual void Shutdown() = 0;
+
+	/**
+	 * Adds a subscription for published messages of the specified type.
+	 *
+	 * Subscriptions allow message consumers to receive published messages from the message bus.
+	 * The returned interface can be used to query the subscription's details and its enabled state.
+	 *
+	 * @param Subscriber The subscriber wishing to receive the messages.
+	 * @param MessageType The type of messages to subscribe to (NAME_All = subscribe to all message types).
+	 * @param ScopeRange The range of message scopes to include in the subscription.
+	 * @return The added subscription, or nullptr if the subscription failed.
+	 * @see Unsubscribe
+	 */
+	virtual TSharedPtr<IMessageSubscription, ESPMode::ThreadSafe> Subscribe( const TSharedRef<IReceiveMessages, ESPMode::ThreadSafe>& Subscriber, const FName& MessageType, const TRange<EMessageScope>& ScopeRange ) = 0;
+
+	/**
+	 * Removes an interceptor for messages of the specified type.
+	 *
+	 * @param Interceptor The interceptor to remove.
+	 * @param MessageType The type of messages to stop intercepting.
+	 * @see Intercept
+	 */
+	virtual void Unintercept( const TSharedRef<IMessageInterceptor, ESPMode::ThreadSafe>& Interceptor, const FName& MessageType ) = 0;
+
+	/**
+	 * Unregisters a message recipient from the message bus.
+	 *
+	 * @param Address The address of the recipient to unregister.
+	 * @see Register
+	 */
+	virtual void Unregister( const FMessageAddress& Address ) = 0;
+
+	/**
+	 * Cancels the specified message subscription.
+	 *
+	 * @param Subscriber The subscriber wishing to stop receiving the messages.
+	 * @param MessageType The type of messages to unsubscribe from (NAME_All = all types).
+	 * @see Subscribe
+	 */
+	virtual void Unsubscribe( const TSharedRef<IReceiveMessages, ESPMode::ThreadSafe>& Subscriber, const FName& MessageType ) = 0;
+
+public:
+
+	/**
+	 * Returns a delegate that is executed when the message bus is shutting down.
+	 *
+	 * @return The delegate.
+	 * @see Shutdown
+	 */
+	virtual FOnMessageBusShutdown& OnShutdown() = 0;
+
+public:
+
+	/** Virtual destructor. */
+	virtual ~IMessageBus() { }
+};
+
+
+/** Type definition for shared pointers to instances of IMessageBus. */
+typedef TSharedPtr<IMessageBus, ESPMode::ThreadSafe> IMessageBusPtr;
+
+/** Type definition for shared references to instances of IMessageBus. */
+typedef TSharedRef<IMessageBus, ESPMode::ThreadSafe> IMessageBusRef;
+
+/** Type definition for weak pointers to instances of IMessageBus. */
+typedef TWeakPtr<IMessageBus, ESPMode::ThreadSafe> IMessageBusWeakPtr;
diff --git a/Engine/Source/Runtime/Messaging/Public/IMessageContext.h b/Engine/Source/Runtime/Messaging/Public/IMessageContext.h
new file mode 100644
--- /dev/null
+++ b/Engine/Source/Runtime/Messaging/Public/IMessageContext.h
@@ -0,0 +1,328 @@
+// Copyright 1998-2014 Epic Games, Inc. All Rights Reserved.
+
+#pragma once
+
+#include "TaskGraphInterfaces.h"
+
+
+// forward declarations
+class IMessageAttachment;
+class UScriptStruct;
+
+
+/**
+ * Structure for message endpoint addresses.
+ */
+struct FMessageAddress
+{
+public:
+
+	/** Default constructor (no initialization). */
+	FMessageAddress() { }
+
+public:
+
+	/**
+	 * Compares two message addresses for equality.
+	 *
+	 * @param X The first address to compare.
+	 * @param Y The second address to compare.
+	 * @return true if the addresses are equal, false otherwise.
+	 */
+	friend bool operator==( const FMessageAddress& X, const FMessageAddress& Y )
+	{
+		return (X.UniqueId == Y.UniqueId);
+	}
+
+	/**
+	 * Compares two message addresses for inequality.
+	 *
+	 * @param X The first address to compare.
+	 * @param Y The second address to compare.
+	 * @return true if the addresses are not equal, false otherwise.
+	 */
+	friend bool operator!=( const FMessageAddress& X, const FMessageAddress& Y )
+	{
+		return (X.UniqueId != Y.UniqueId);
+	}
+
+	/**
+	 * Serializes a message address from or into an archive.
+	 *
+	 * @param Ar The archive to serialize from or into.
+	 * @param G The address to serialize.
+	 */
+	friend FArchive& operator<<( FArchive& Ar, FMessageAddress& A )
+	{
+		return Ar << A.UniqueId;
+	}
+
+public:
+
+	/**
+	 * Invalidates the GUID.
+	 *
+	 * @see IsValid
+	 */
+	void Invalidate()
+	{
+		UniqueId.Invalidate();
+	}
+
+	/**
+	 * Checks whether this message address is valid or not.
+	 *
+	 * @return true if valid, false otherwise.
+	 * @see Invalidate
+	 */
+	bool IsValid() const
+	{
+		return UniqueId.IsValid();
+	}
+
+	/**
+	 * Converts this GUID to its string representation.
+	 *
+	 * @return The string representation.
+	 */
+	FString ToString() const
+	{
+		return UniqueId.ToString();
+	}
+
+public:
+
+	/**
+	 * Calculates the hash for a message address.
+	 *
+	 * @param Address The address to calculate the hash for.
+	 * @return The hash.
+	 */
+	friend uint32 GetTypeHash( const FMessageAddress& Address )
+	{
+		return FCrc::MemCrc_DEPRECATED(&Address.UniqueId, sizeof(FGuid));
+	}
+
+public:
+
+	/**
+	 * Returns a new message address.
+	 *
+	 * @return A new address.
+	 */
+	static FMessageAddress NewAddress()
+	{
+		FMessageAddress Result;
+		Result.UniqueId = FGuid::NewGuid();
+
+		return Result;
+	}
+
+private:
+
+	/** Holds a unique identifier. */
+	FGuid UniqueId;
+};
+
+
+/**
+ * Enumerates scopes for published messages.
+ *
+ * The scope determines to which endpoints a published message will be delivered.
+ * By default, messages will be published to everyone on the local network (Subnet),
+ * but it is often useful to restrict the group of recipients to more local scopes,
+ * or to widen it to a larger audience outside the local network.
+ *
+ * For example, if a message is to be handled only by subscribers in the same application,
+ * the scope should be set to Process. If messages need to be published between
+ * different networks (i.e. between LAN and WLAN), it should be set to Network instead.
+ * 
+ * Scopes only apply to published messages. Messages that are being sent to specific
+ * recipients will always be delivered, regardless of the endpoint locations.
+ */
+enum class EMessageScope
+{
+	/** Deliver to subscribers in the same thread. */
+	Thread,
+
+	/** Deliver to subscribers in the same process. */
+	Process,
+
+	/** Deliver to subscribers on the network. */
+	Network,
+
+	/**
+	 * Deliver to all subscribers.
+	 *
+	 * Note: This must be the last value in this enumeration.
+	 */
+	All
+};
+
+
+/** Type definition for message scope ranges. */
+typedef TRange<EMessageScope> FMessageScopeRange;
+
+/** Type definition for message scope range bounds. */
+typedef TRangeBound<EMessageScope> FMessageScopeRangeBound;
+
+
+/**
+ * Interface for message contexts.
+ *
+ * Messages are delivered inside message contexts, which store the message itself plus additional data
+ * associated with the message. Recipients of a message are usually interested in data that describes the
+ * message, such as its origin or when it expires. They may also be interested in optional out-of-band
+ * binary data that is attached to the message.
+ *
+ * The sender's address (@see GetSender) is often needed to send a reply message to a message sender, i.e.
+ * in response to a published message. The message attachment (@see GetAttchment) is an optional bundle
+ * of binary bulk data that is transferred independently from the message itself and allows for transferring
+ * larger amounts of data that would otherwise clog up the messaging system.
+ *
+ * In case a message was forwarded by another endpoint, the context of the original sender can be accessed
+ * using the @see GetOriginalContext method.
+ *
+ * @see IMessageAttachment
+ */
+class IMessageContext
+{
+public:
+
+	/**
+	 * Gets the optional message annotations.
+	 *
+	 * @return Message header collection.
+	 */
+	virtual const TMap<FName, FString>& GetAnnotations() const = 0;
+
+	/**
+	 * Gets the message attachment, if present.
+	 *
+	 * @return A pointer to the message attachment, or nullptr if no attachment is present.
+	 */
+	virtual TSharedPtr<IMessageAttachment, ESPMode::ThreadSafe> GetAttachment() const = 0;
+
+	/**
+	 * Gets the date and time at which the message expires.
+	 *
+	 * @return Expiration time.
+	 */
+	virtual const FDateTime& GetExpiration() const = 0;
+
+	/**
+	 * Gets the message data.
+	 *
+	 * @return A pointer to the message data.
+	 * @see GetMessageType, GetMessageTypeInfo
+	 */
+	virtual const void* GetMessage() const = 0;
+
+	/**
+	 * Gets the message's type information.
+	 *
+	 * @return Message type information.
+	 * @see GetMessage, GetMessageType
+	 */
+	virtual const TWeakObjectPtr<UScriptStruct>& GetMessageTypeInfo() const = 0;
+
+	/**
+	 * Returns the original message context in case the message was forwarded.
+	 *
+	 * @return The original message context, or nullptr if the message wasn't forwarded.
+	 */
+	virtual TSharedPtr<IMessageContext, ESPMode::ThreadSafe> GetOriginalContext() const = 0;
+
+	/**
+	 * Gets the list of message recipients.
+	 *
+	 * @return Message recipients.
+	 */
+	virtual const TArray<FMessageAddress>& GetRecipients() const = 0;
+
+	/**
+	 * Gets the scope to which the message was sent.
+	 *
+	 * @return The message scope.
+	 */
+	virtual EMessageScope GetScope() const = 0;
+
+	/**
+	 * Gets the sender's address.
+	 *
+	 * @return Sender address.
+	 */
+	virtual const FMessageAddress& GetSender() const = 0;
+
+	/**
+	 * Gets the name of the thread from which the message was sent.
+	 *
+	 * @return Sender threat name.
+	 */
+	virtual ENamedThreads::Type GetSenderThread() const = 0;
+
+	/**
+	 * Gets the time at which the message was forwarded.
+	 *
+	 * @return Time forwarded.
+	 */
+	virtual const FDateTime& GetTimeForwarded() const = 0;
+
+	/**
+	 * Gets the time at which the message was sent.
+	 *
+	 * @return Time sent.
+	 */
+	virtual const FDateTime& GetTimeSent() const = 0;
+
+public:
+
+	/**
+	 * Gets the name of the message type.
+	 *
+	 * @return Message type name.
+	 */
+	FName GetMessageType() const
+	{
+		if (IsValid())
+		{
+			return GetMessageTypeInfo()->GetFName();
+		}
+		
+		return NAME_None;
+	}
+
+	/**
+	 * Checks whether this is a forwarded message.
+	 *
+	 * @return true if the message was forwarded, false otherwise.
+	 * @see GetOriginalContext, IsValid
+	 */
+	bool IsForwarded() const
+	{
+		return GetOriginalContext().IsValid();
+	}
+
+	/**
+	 * Checks whether this context is valid.
+	 *
+	 * @return true if the context is valid, false otherwise.
+	 * @see IsForwarded
+	 */
+	bool IsValid() const
+	{
+		return ((GetMessage() != nullptr) && GetMessageTypeInfo().IsValid(false, true));
+	}
+
+public:
+
+	/** Virtual destructor. */
+	virtual ~IMessageContext() { }
+};
+
+
+/** Type definition for shared pointers to instances of IMessageContext. */
+typedef TSharedPtr<IMessageContext, ESPMode::ThreadSafe> IMessageContextPtr;
+
+/** Type definition for shared references to instances of IMessageContext. */
+typedef TSharedRef<IMessageContext, ESPMode::ThreadSafe> IMessageContextRef;
diff --git a/Engine/Source/Runtime/Messaging/Public/IMessageHandler.h b/Engine/Source/Runtime/Messaging/Public/IMessageHandler.h
new file mode 100644
--- /dev/null
+++ b/Engine/Source/Runtime/Messaging/Public/IMessageHandler.h
@@ -0,0 +1,108 @@
+// Copyright 1998-2014 Epic Games, Inc. All Rights Reserved.
+
+#pragma once
+
+#include "CoreUObject.h"
+
+
+// forward declarations
+class IMessageContext;
+
+
+/**
+ * Interface for message handlers.
+ */
+class IMessageHandler
+{
+public:
+
+	/**
+	 * Gets the message type handled by this handler.
+	 *
+	 * @return The name of the message type.
+	 * @see HandleMessage
+	 */
+	virtual const FName GetHandledMessageType() const = 0;
+
+	/**
+	 * Handles the specified message.
+	 *
+	 * @param Context The context of the message to handle.
+	 * @see GetHandledMessageType
+	 */
+	virtual void HandleMessage( const TSharedRef<IMessageContext, ESPMode::ThreadSafe>& Context ) = 0;
+
+public:
+
+	/** Virtual destructor. */
+	virtual ~IMessageHandler() { }
+};
+
+
+/** Type definition for shared pointers to instances of IMessageHandler. */
+typedef TSharedPtr<IMessageHandler, ESPMode::ThreadSafe> IMessageHandlerPtr;
+
+/** Type definition for shared references to instances of IMessageHandler. */
+typedef TSharedRef<IMessageHandler, ESPMode::ThreadSafe> IMessageHandlerRef;
+
+
+/**
+ * Template for message handler function pointer types.
+ *
+ * @param MessageType The type of message to handle.
+ * @param HandlerType The type of the handler class.
+ */
+template<typename MessageType, typename HandlerType>
+struct TMessageHandlerFunc
+{
+	typedef void (HandlerType::*Type)( const MessageType&, const TSharedRef<IMessageContext, ESPMode::ThreadSafe>& );
+};
+
+
+/**
+ * Template for message handlers.
+ *
+ * @param MessageType The type of message to handle.
+ * @param HandlerType The type of the handler class.
+ */
+template<typename MessageType, typename HandlerType>
+class TMessageHandler
+	: public IMessageHandler
+{
+public:
+
+	/**
+	 * Creates and initializes a new message handler.
+	 *
+	 * @param InHandler The object that will handle the messages.
+	 * @param InHandlerFunc The object's message handling function.
+	 */
+	TMessageHandler( HandlerType* InHandler, typename TMessageHandlerFunc<MessageType, HandlerType>::Type InHandlerFunc )
+		: Handler(InHandler)
+		, HandlerFunc(InHandlerFunc)
+	{
+		check(InHandler != nullptr);
+	}
+
+public:
+
+	// IMessageHandler interface
+	
+	virtual const FName GetHandledMessageType() const override
+	{
+		return MessageType::StaticStruct()->GetFName();
+	}
+
+	virtual void HandleMessage( const TSharedRef<IMessageContext, ESPMode::ThreadSafe>& Context ) override
+	{
+		(Handler->*HandlerFunc)(*static_cast<const MessageType*>(Context->GetMessage()), Context);
+	}
+	
+private:
+
+	/** Holds a pointer to the object handling the messages. */
+	HandlerType* Handler;
+
+	/** Holds a pointer to the actual handler function. */
+	typename TMessageHandlerFunc<MessageType, HandlerType>::Type HandlerFunc;
+};
diff --git a/Engine/Source/Runtime/Messaging/Public/IMessageInterceptor.h b/Engine/Source/Runtime/Messaging/Public/IMessageInterceptor.h
new file mode 100644
--- /dev/null
+++ b/Engine/Source/Runtime/Messaging/Public/IMessageInterceptor.h
@@ -0,0 +1,36 @@
+// Copyright 1998-2014 Epic Games, Inc. All Rights Reserved.
+
+#pragma once
+
+
+// forward declarations
+class IMessageContext;
+
+
+/**
+ * Interface for message interceptors.
+ */
+class IMessageInterceptor
+{
+public:
+
+	/**
+	 * Intercepts a message before it is being passed to the message router.
+	 *
+	 * @param Context The context of the message to intercept.
+	 * @return true if the message was intercepted and should not be routed, false otherwise.
+	 */
+	virtual bool InterceptMessage( const TSharedRef<IMessageContext, ESPMode::ThreadSafe>& Context ) = 0;
+
+public:
+
+	/** Virtual destructor. */
+	virtual ~IMessageInterceptor() { }
+};
+
+
+/** Type definition for shared pointers to instances of IInterceptMessages. */
+typedef TSharedPtr<IMessageInterceptor, ESPMode::ThreadSafe> IMessageInterceptorPtr;
+
+/** Type definition for shared references to instances of IInterceptMessages. */
+typedef TSharedRef<IMessageInterceptor, ESPMode::ThreadSafe> IMessageInterceptorRef;
diff --git a/Engine/Source/Runtime/Messaging/Public/IMessageSubscription.h b/Engine/Source/Runtime/Messaging/Public/IMessageSubscription.h
new file mode 100644
--- /dev/null
+++ b/Engine/Source/Runtime/Messaging/Public/IMessageSubscription.h
@@ -0,0 +1,83 @@
+// Copyright 1998-2014 Epic Games, Inc. All Rights Reserved.
+
+#pragma once
+
+
+// forward declarations
+enum class EMessageScope;
+class IReceiveMessages;
+
+
+/**
+ * Interface for message subscriptions.
+ *
+ * This interface provides access to a message subscription that was previously created with the
+ * @see IMessageBus.Subscribe method. It can be used to query the subscription's details using the
+ * various getters and modify its enabled state using the @see IMessageSubscription.Enable and
+ * IMessageSubscription.Disable methods.
+ *
+ * A subscription that is disabled will cause a message endpoint to stop receiving messages for
+ * the subscribed message type. However, the subscription is merely disabled and not removed.
+ */
+class IMessageSubscription
+{
+public:
+
+	/**
+	 * Disables the subscription.
+	 *
+	 * @see Enable, IsEnabled
+	 */
+	virtual void Disable() = 0;
+
+	/**
+	 * Enables the subscription.
+	 *
+	 * @see Disable, IsEnabled
+	 */
+	virtual void Enable() = 0;
+
+	/**
+	 * Gets the type of subscribed messages.
+	 *
+	 * @return Message type.
+	 * @see GetScopeRange, GetSubscriber
+	 */
+	virtual FName GetMessageType() = 0;
+
+	/**
+	 * Gets the range of subscribed message scopes.
+	 *
+	 * @return Message scope range.
+	 * @see GetMessageType, GetSubscriber
+	 */
+	virtual const TRange<EMessageScope>& GetScopeRange() = 0;
+
+	/**
+	 * Gets the subscriber.
+	 *
+	 * @return The subscriber.
+	 * @see GetMessageType, GetScopeRange
+	 */
+	virtual const TWeakPtr<IReceiveMessages, ESPMode::ThreadSafe>& GetSubscriber() = 0;
+
+	/**
+	 * Checks whether the subscription is enabled.
+	 *
+	 * @return true if the subscription is enabled, false otherwise.
+	 * @see Enable, Disable
+	 */
+	virtual bool IsEnabled() = 0;
+
+public:
+
+	/** Virtual destructor. */
+	virtual ~IMessageSubscription() { }
+};
+
+
+/** Type definition for shared pointers to instances of IMessageSubscription. */
+typedef TSharedPtr<IMessageSubscription, ESPMode::ThreadSafe> IMessageSubscriptionPtr;
+
+/** Type definition for shared references to instances of IMessageSubscription. */
+typedef TSharedRef<IMessageSubscription, ESPMode::ThreadSafe> IMessageSubscriptionRef;
diff --git a/Engine/Source/Runtime/Messaging/Public/IMessageTracer.h b/Engine/Source/Runtime/Messaging/Public/IMessageTracer.h
new file mode 100644
--- /dev/null
+++ b/Engine/Source/Runtime/Messaging/Public/IMessageTracer.h
@@ -0,0 +1,311 @@
+// Copyright 1998-2014 Epic Games, Inc. All Rights Reserved.
+
+#pragma once
+
+
+// forward declarations
+class IMessageContext;
+
+
+/** Type definition for shared pointers to instances of FMessageTracerAddressInfo. */
+typedef TSharedPtr<struct FMessageTracerAddressInfo> FMessageTracerAddressInfoPtr;
+
+/** Type definition for shared references to instances of FMessageTracerAddressInfo. */
+typedef TSharedRef<struct FMessageTracerAddressInfo> FMessageTracerAddressInfoRef;
+
+/** Type definition for shared pointers to instances of FMessageTracerDispatchState. */
+typedef TSharedPtr<struct FMessageTracerDispatchState> FMessageTracerDispatchStatePtr;
+
+/** Type definition for shared references to instances of FMessageTracerDispatchState. */
+typedef TSharedRef<struct FMessageTracerDispatchState> FMessageTracerDispatchStateRef;
+
+/** Type definition for shared pointers to instances of FMessageTracerEndpointInfo. */
+typedef TSharedPtr<struct FMessageTracerEndpointInfo> FMessageTracerEndpointInfoPtr;
+
+/** Type definition for shared references to instances of FMessageTracerEndpointInfo. */
+typedef TSharedRef<struct FMessageTracerEndpointInfo> FMessageTracerEndpointInfoRef;
+
+/** Type definition for shared pointers to instances of FMessageTracerMessageInfo. */
+typedef TSharedPtr<struct FMessageTracerMessageInfo> FMessageTracerMessageInfoPtr;
+
+/** Type definition for shared references to instances of FMessageTracerMessageInfo. */
+typedef TSharedRef<struct FMessageTracerMessageInfo> FMessageTracerMessageInfoRef;
+
+/** Type definition for shared pointers to instances of FMessageTracerTypeInfo. */
+typedef TSharedPtr<struct FMessageTracerTypeInfo> FMessageTracerTypeInfoPtr;
+
+/** Type definition for shared references to instances of FMessageTracerTypeInfo. */
+typedef TSharedRef<struct FMessageTracerTypeInfo> FMessageTracerTypeInfoRef;
+
+
+/**
+ * Enumerates message dispatch types.
+ */
+enum class EMessageTracerDispatchTypes
+{
+	/** The message is being dispatched directly. */
+	Direct,
+
+	/** The message hasn't been dispatched yet. */
+	Pending,
+
+	/** The message is being dispatched using the task graph system. */
+	TaskGraph
+};
+
+
+/**
+ * Structure for message dispatch states.
+ */
+struct FMessageTracerDispatchState
+{
+	/** Holds the dispatch latency (in seconds). */
+	double DispatchLatency;
+
+	/** Holds the message's dispatch type for the specified endpoint. */
+	EMessageTracerDispatchTypes DispatchType;
+
+	/** Holds the endpoint to which the message was or is being dispatched. */
+	FMessageTracerEndpointInfoPtr EndpointInfo;
+
+	/** Holds the time at which the message was dispatched. */
+	double TimeDispatched;
+
+	/** Holds the time at which the message was actually handled (0.0 = not handled yet). */
+	double TimeHandled;
+};
+
+
+/**
+ * Structure for a recipient's address information.
+ */
+struct FMessageTracerAddressInfo
+{
+	/** Holds a recipient's message address. */
+	FMessageAddress Address;
+
+	/** Holds the time at which this address was registered. */
+	double TimeRegistered;
+
+	/** Holds the time at which this address was unregistered. */
+	double TimeUnregistered;
+};
+
+
+/**
+ * Structure for message endpoint debug information.
+ */
+struct FMessageTracerEndpointInfo
+{
+	/** Holds the recipient's address information. */
+	TMap<FMessageAddress, FMessageTracerAddressInfoPtr> AddressInfos;
+
+	/** Holds the recipient's human readable name. */
+	FName Name;
+
+	/** Holds the list of messages received by this recipient. */
+	TArray<FMessageTracerMessageInfoPtr> ReceivedMessages;
+
+	/** Holds a flag indicating whether this is a remote recipient. */
+	bool Remote;
+
+	/** Holds the list of messages sent by this recipient. */
+	TArray<FMessageTracerMessageInfoPtr> SentMessages;
+};
+
+
+/**
+ * Structure for message debug information.
+ */
+struct FMessageTracerMessageInfo
+{
+	/** Holds a pointer to the message context. */
+	IMessageContextPtr Context;
+
+	/** Holds the message's dispatch states per endpoint. */
+	TMap<FMessageTracerEndpointInfoPtr, FMessageTracerDispatchStatePtr> DispatchStates;
+
+	/** Pointer to the sender's endpoint information. */
+	FMessageTracerEndpointInfoPtr SenderInfo;
+
+	/** Holds the time at which the message was routed (0.0 = pending). */
+	double TimeRouted;
+
+	/** Holds the time at which the message was sent. */
+	double TimeSent;
+
+	/** Pointer to the message's type information. */
+	FMessageTracerTypeInfoPtr TypeInfo;
+};
+
+
+/**
+ * Structure for message type debug information.
+ */
+struct FMessageTracerTypeInfo
+{
+	/** Holds the collection of messages of this type. */
+	TArray<FMessageTracerMessageInfoPtr> Messages;
+
+	/** Holds a name of the message type. */
+	FName TypeName;
+};
+
+
+/**
+ * Implements a delegate that is executed when a message has been added to the collection of traced messages.
+ *
+ * The first parameter is the information for the message that was added.
+ */
+DECLARE_MULTICAST_DELEGATE_OneParam(FMessageTracerMessageAdded, FMessageTracerMessageInfoRef)
+
+/**
+ * Implements a delegate that is executed when a type has been added to the collection of traced message types.
+ *
+ * The first parameter is the information for the message type that was added.
+ */
+DECLARE_MULTICAST_DELEGATE_OneParam(FMessageTracerTypeAdded, FMessageTracerTypeInfoRef)
+
+
+/**
+ * Interface for message tracers.
+ *
+ * This interface provides access to the message tracer that is built into a message bus. Message tracers are
+ * able to monitor and record all activities on a message bus. This includes the sending and receiving of messages,
+ * the registering and unregistering of message endpoints and interceptors and adding and removal of subscriptions.
+ */
+class IMessageTracer
+{
+public:
+
+	/**
+	 * Breaks message routing.
+	 *
+	 * @see Continue, IsBreaking, Step
+	 */
+	virtual void Break() = 0;
+
+	/**
+	 * Continues message routing from the current breakpoint.
+	 *
+	 * @see Break, Step
+	 */
+	virtual void Continue() = 0;
+
+	/**
+	 * Checks whether the tracer is currently at a breakpoint.
+	 *
+	 * @return true if at breakpoint, false otherwise.
+	 * @see Break
+	 */
+	virtual bool IsBreaking() const = 0;
+
+	/**
+	 * Checks whether the tracer is currently running.
+	 *
+	 * @return true if the tracer is running, false otherwise.
+	 *
+	 * @see Start, Stop
+	 */
+	virtual bool IsRunning() const = 0;
+
+	/** Resets the tracer. */
+	virtual void Reset() = 0;
+
+	/**
+	 * Starts the tracer.
+	 *
+	 * @see IsRunning, Stop
+	 */
+	virtual void Start() = 0;
+
+	/**
+	 * Steps the tracer to the next message.
+	 *
+	 * @see Break, Continue
+	 */
+	virtual void Step() = 0;
+
+	/**
+	 * Stops the tracer.
+	 *
+	 * @see IsRunning, Start
+	 */
+	virtual void Stop() = 0;
+
+	/**
+	 * Ticks the tracer.
+	 *
+	 * @param DeltaTime The time in seconds since the last tick.
+	 * @return true if any events were processed.
+	 */
+	virtual bool Tick( float DeltaTime ) = 0;
+
+public:
+
+	/**
+	 * Gets the list of known message endpoints.
+	 *
+	 * @param OutEndpoints Will contain the list of endpoints.
+	 * @return The number of endpoints returned.
+	 */
+	virtual int32 GetEndpoints( TArray<FMessageTracerEndpointInfoPtr>& OutEndpoints ) const = 0;
+
+	/**
+	 * Gets the collection of known messages.
+	 *
+	 * @return The messages.
+	 */
+	virtual int32 GetMessages( TArray<FMessageTracerMessageInfoPtr>& OutMessages ) const = 0;
+
+	/**
+	 * Gets the list of known message types filtered by name.
+	 *
+	 * @param NameFilter The name substring to filter with.
+	 * @param OutTypes Will contain the list of message types.
+	 * @return The number of message types returned.
+	 */
+	virtual int32 GetMessageTypes( TArray<FMessageTracerTypeInfoPtr>& OutTypes ) const = 0;
+
+	/**
+	 * Checks whether there are any messages in the history.
+	 *
+	 * @return true if there are messages, false otherwise.
+	 */
+	virtual bool HasMessages() const = 0;
+
+public:
+
+	/**
+	 * Returns a delegate that is executed when the collection of known messages has changed.
+	 *
+	 * @return The delegate.
+	 */
+	virtual FMessageTracerMessageAdded& OnMessageAdded() = 0;
+
+	/**
+	 * Returns a delegate that is executed when the message history has been reset.
+	 *
+	 * @return The delegate.
+	 */
+	virtual FSimpleMulticastDelegate& OnMessagesReset() = 0;
+
+	/**
+	 * Returns a delegate that is executed when the collection of known messages types has changed.
+	 *
+	 * @return The delegate.
+	 */
+	virtual FMessageTracerTypeAdded& OnTypeAdded() = 0;
+
+protected:
+
+	/** Hidden destructor. The life time of a message tracer is managed by the message bus. */
+	virtual ~IMessageTracer() { }
+};
+
+
+/** Type definition for shared pointers to instances of IMessageTracer. */
+typedef TSharedPtr<IMessageTracer, ESPMode::ThreadSafe> IMessageTracerPtr;
+
+/** Type definition for shared references to instances of IMessageTracer. */
+typedef TSharedRef<IMessageTracer, ESPMode::ThreadSafe> IMessageTracerRef;
diff --git a/Engine/Source/Runtime/Messaging/Public/IMessageTracerBreakpoint.h b/Engine/Source/Runtime/Messaging/Public/IMessageTracerBreakpoint.h
new file mode 100644
--- /dev/null
+++ b/Engine/Source/Runtime/Messaging/Public/IMessageTracerBreakpoint.h
@@ -0,0 +1,45 @@
+// Copyright 1998-2014 Epic Games, Inc. All Rights Reserved.
+
+#pragma once
+
+
+// forward declarations
+class IMessageContext;
+
+
+/**
+ * Interface for message tracer breakpoints.
+ *
+ * @see IMessageTracer
+ */
+class IMessageTracerBreakpoint
+{
+public:
+
+	/**
+	 * Checks whether this breakpoint is enabled.
+	 *
+	 * @return true if the breakpoint is enabled, false otherwise.
+	 */
+	virtual bool IsEnabled() const = 0;
+
+	/**
+	 * Checks whether the tracer should break on the given message.
+	 *
+	 * @param Context The context of the message to break on.
+	 * @return true if the tracer should break, false otherwise.
+	 */
+	virtual bool ShouldBreak( const TSharedRef<IMessageContext, ESPMode::ThreadSafe>& Context ) const = 0;
+
+protected:
+
+	/** Hidden destructor. */
+	virtual ~IMessageTracerBreakpoint() { }
+};
+
+
+/** Type definition for shared pointers to instances of IMessageTracerBreakpoint. */
+typedef TSharedPtr<IMessageTracerBreakpoint, ESPMode::ThreadSafe> IMessageTracerBreakpointPtr;
+
+/** Type definition for shared references to instances of IMessageTracerBreakpoint. */
+typedef TSharedRef<IMessageTracerBreakpoint, ESPMode::ThreadSafe> IMessageTracerBreakpointRef;
diff --git a/Engine/Source/Runtime/Messaging/Public/IMessageTransport.h b/Engine/Source/Runtime/Messaging/Public/IMessageTransport.h
new file mode 100644
--- /dev/null
+++ b/Engine/Source/Runtime/Messaging/Public/IMessageTransport.h
@@ -0,0 +1,112 @@
+// Copyright 1998-2014 Epic Games, Inc. All Rights Reserved.
+
+#pragma once
+
+
+// forward declarations
+class IMessageAttachment;
+class IMessageData;
+
+
+/**
+ * Delegate type for receiving transport messages.
+ *
+ * The first parameter is the context of the received message.
+ * The second parameter is the identifier of the transport node that sent the message.
+ */
+DECLARE_DELEGATE_TwoParams(FOnMessageTransportMessageReceived, const IMessageContextRef&, const FGuid&);
+
+/**
+ * Delegate type for handling the discovery of transport nodes.
+ *
+ * The first parameter is the identifier of the transport node that was discovered.
+ */
+DECLARE_DELEGATE_OneParam(FOnMessageTransportNodeDiscovered, const FGuid&);
+
+/**
+ * Delegate type for handling disconnected or timed out transport nodes.
+ *
+ * Remote nodes can be lost if they disconnect from the bus or if a keep-alive
+ * message was not received within the required timeout (i.e. in case of a crash).
+ *
+ * The first parameter is the identifier of the transport node that was lost.
+ */
+DECLARE_DELEGATE_OneParam(FOnMessageTransportNodeLost, const FGuid&);
+
+
+/**
+ * Interface for message transport technologies.
+ *
+ * Licensees can implement this interface to add support for custom message transport
+ * technologies that are not supported out of the box, i.e. custom network protocols or APIs.
+ */
+class IMessageTransport
+{
+public:
+
+	/**
+	 * Gets the name of this transport (for debugging purposes).
+	 *
+	 * @return The debug name.
+	 */
+	virtual FName GetDebugName() const = 0;
+
+	/**
+	 * Starts up the message transport.
+	 *
+	 * @return Whether the transport was started successfully.
+	 * @see StopTransport
+	 */
+	virtual bool StartTransport() = 0;
+
+	/**
+	 * Shuts down the message transport.
+	 *
+	 * @see StartTransport
+	 */
+	virtual void StopTransport() = 0;
+
+	/**
+	 * Transports the given message data to the specified network nodes.
+	 *
+	 * @param Context The context of the message to transport.
+	 * @param Recipients The transport nodes to send the message to.
+	 * @return true if the message is being transported, false otherwise.
+	 */
+	virtual bool TransportMessage( const IMessageContextRef& Context, const TArray<FGuid>& Recipients ) = 0;
+
+public:
+
+	/**
+	 * Returns a delegate that is executed when message data has been received.
+	 *
+	 * @param Delegate The delegate to set.
+	 */
+	virtual FOnMessageTransportMessageReceived& OnMessageReceived() = 0;
+
+	/**
+	 * Returns a delegate that is executed when a transport node has been discovered.
+	 *
+	 * @param Delegate The delegate to set.
+	 */
+	virtual FOnMessageTransportNodeDiscovered& OnNodeDiscovered() = 0;
+
+	/**
+	 * Returns a delegate that is executed when a transport node has closed or timed out.
+	 *
+	 * @param Delegate The delegate to set.
+	 */
+	virtual FOnMessageTransportNodeLost& OnNodeLost() = 0;
+
+protected:
+
+	/** Virtual constructor. */
+	~IMessageTransport() { }
+};
+
+
+/** Type definition for shared pointers to instances of ITransportMessages. */
+typedef TSharedPtr<IMessageTransport, ESPMode::ThreadSafe> IMessageTransportPtr;
+
+/** Type definition for shared references to instances of ITransportMessages. */
+typedef TSharedRef<IMessageTransport, ESPMode::ThreadSafe> IMessageTransportRef;
diff --git a/Engine/Source/Runtime/Messaging/Public/IMessagingModule.h b/Engine/Source/Runtime/Messaging/Public/IMessagingModule.h
new file mode 100644
--- /dev/null
+++ b/Engine/Source/Runtime/Messaging/Public/IMessagingModule.h
@@ -0,0 +1,78 @@
+// Copyright 1998-2014 Epic Games, Inc. All Rights Reserved.
+
+#pragma once
+
+#include "ModuleInterface.h"
+
+
+// forward declarations
+struct FMessageAddress;
+class IAuthorizeMessageRecipients;
+class IMessageBus;
+class IMessageBridge;
+class IMessageTransport;
+
+
+/**
+ * Interface for messaging modules.
+ *
+ * @see IMessageBridge, IMessageBus
+ */
+class IMessagingModule
+	: public IModuleInterface
+{
+public:
+
+	/**
+	 * Creates a new message bridge.
+	 *
+	 * Message bridges translate messages between a message bus and another means of
+	 * message transportation, such as network sockets.
+	 *
+	 * @param Address The bridge's address on the message bus.
+	 * @param Bus The message bus to attach the bridge to.
+	 * @param Transport The message transport technology to use.
+	 * @return The new message bridge, or nullptr if the bridge couldn't be created.
+	 * @see CreateBus
+	 */
+	virtual TSharedPtr<IMessageBridge, ESPMode::ThreadSafe> CreateBridge( const FMessageAddress& Address, const TSharedRef<IMessageBus, ESPMode::ThreadSafe>& Bus, const TSharedRef<IMessageTransport, ESPMode::ThreadSafe>& Transport ) = 0;
+
+	/**
+	 * Creates a new message bus.
+	 *
+	 * @param RecipientAuthorizer An optional recipient authorizer.
+	 * @return The new message bus, or nullptr if the bus couldn't be created.
+	 * @see CreateBridge
+	 */
+	virtual TSharedPtr<IMessageBus, ESPMode::ThreadSafe> CreateBus( const TSharedPtr<IAuthorizeMessageRecipients>& RecipientAuthorizer ) = 0;
+
+	/**
+	 * Gets the default message bus if it has been initialized.
+	 *
+	 * @return The default bus.
+	 */
+	virtual TSharedPtr<IMessageBus, ESPMode::ThreadSafe> GetDefaultBus() const = 0;
+
+public:
+
+	/**
+	 * Gets a reference to the messaging module instance.
+	 *
+	 * @return A reference to the Messaging module.
+	 * @todo gmp: better implementation using dependency injection.
+	 */
+	static IMessagingModule& Get()
+	{
+#if PLATFORM_IOS
+        static IMessagingModule& MessageModule = FModuleManager::LoadModuleChecked<IMessagingModule>("Messaging");
+        return MessageModule;
+#else
+        return FModuleManager::LoadModuleChecked<IMessagingModule>("Messaging");
+#endif
+	}
+
+public:
+
+	/** Virtual destructor. */
+	virtual ~IMessagingModule() { }
+};
diff --git a/Engine/Source/Runtime/Messaging/Public/IMutableMessageContext.h b/Engine/Source/Runtime/Messaging/Public/IMutableMessageContext.h
new file mode 100644
--- /dev/null
+++ b/Engine/Source/Runtime/Messaging/Public/IMutableMessageContext.h
@@ -0,0 +1,90 @@
+// Copyright 1998-2014 Epic Games, Inc. All Rights Reserved.
+
+#pragma once
+
+#include "IMessageContext.h"
+
+
+// forward declarations
+class IMessageAttachment;
+class UScriptStruct;
+
+
+/**
+ * Interface for mutable message contexts.
+ */
+class IMutableMessageContext
+	: public IMessageContext
+{
+public:
+
+	/**
+	 * Adds a message address to the recipient list.
+	 *
+	 * @param Recipient The address of the recipient to add.
+	 */
+	virtual void AddRecipient( const FMessageAddress& Recipient ) = 0;
+
+	/**
+	 * Sets the optional message attachment.
+	 *
+	 * @param InAttachment The attachment to set.
+	 */
+	virtual void SetAttachment( const TSharedPtr<IMessageAttachment, ESPMode::ThreadSafe>& InAttachment ) = 0;
+
+	/**
+	 * Sets the message.
+	 *
+	 * @param InMessage The message to set.
+	 * @param InTypeInfo The message's type information.
+	 */
+	virtual void SetMessage( void* InMessage, UScriptStruct* InTypeInfo ) = 0;
+
+	/**
+	 * Sets the date and time at which the message expires.
+	 *
+	 * @param InExpiration Expiration date and time.
+	 */
+	virtual void SetExpiration( const FDateTime& InExpiration ) = 0;
+
+	/**
+	 * Sets the value of the header with the specified key.
+	 *
+	 * @param Key The header key.
+	 * @param Value The header value.
+	 */
+	virtual void SetHeader( const FName& Key, const FString& Value ) = 0;
+
+	/**
+	 * Sets the message scope.
+	 *
+	 * @param InScope The message scope.
+	 */
+	virtual void SetScope( EMessageScope InScope ) = 0;
+
+	/**
+	 * Sets the address of the message's sender.
+	 *
+	 * @param InSender The message sender's address.
+	 */
+	virtual void SetSender( const FMessageAddress& InSender ) = 0;
+
+	/**
+	 * Sets the date and time at which the message was sent.
+	 *
+	 * @param InTimeSent Send date and time.
+	 */
+	virtual void SetTimeSent( const FDateTime& InTimeSent ) = 0;
+
+public:
+
+	/** Virtual destructor. */
+	virtual ~IMutableMessageContext() { }
+};
+
+
+/** Type definition for shared pointers to instances of IMutableMessageContext. */
+typedef TSharedPtr<IMutableMessageContext, ESPMode::ThreadSafe> IMutableMessageContextPtr;
+
+/** Type definition for shared references to instances of IMutableMessageContext. */
+typedef TSharedRef<IMutableMessageContext, ESPMode::ThreadSafe> IMutableMessageContextRef;
diff --git a/Engine/Source/Runtime/Messaging/Public/IReceiveMessages.h b/Engine/Source/Runtime/Messaging/Public/IReceiveMessages.h
new file mode 100644
--- /dev/null
+++ b/Engine/Source/Runtime/Messaging/Public/IReceiveMessages.h
@@ -0,0 +1,100 @@
+// Copyright 1998-2014 Epic Games, Inc. All Rights Reserved.
+
+#pragma once
+
+
+// forward declarations
+class IMessageContext;
+
+
+/**
+ * Interface for message recipients.
+ *
+ * Classes that implement this interface are able to receive messages from a message bus. A message recipient will receive
+ * a call to its ReceiveMessage() method for each message that was sent directly to it (@see IMessageBus.Send) and for each
+ * published message that it subscribed to (@see IMessageBus.Publish).
+ *
+ * This interface provides a rather low-level mechanism for receiving messages. Instead of implementing it, Most users will
+ * want to use an instance of @see FMessageEndpoint, which provides a much more convenient way of sending and receiving messages.
+ *
+ * @see FMessageEndpoint, IMessageBus, ISendMessages
+ */
+class IReceiveMessages
+{
+public:
+
+	/**
+	 * Gets the recipient's name (for debugging purposes).
+	 *
+	 * @return The debug name.
+	 * @see GetRecipientId, GetRecipientThread
+	 */
+	virtual FName GetDebugName() const = 0;
+
+	/**
+	 * Gets the recipient's unique identifier (for debugging purposes).
+	 *
+	 * @return The recipient's identifier.
+	 * @see GetDebugName, GetRecipientThread
+	 */
+	virtual const FGuid& GetRecipientId() const = 0;
+
+	/**
+	 * Gets the name of the thread on which to receive messages.
+	 *
+	 * If the recipient's ReceiveMessage() is thread-safe, return ThreadAny for best performance.
+	 *
+	 * @return Name of the receiving thread.
+	 * @see GetDebugName, GetRecipientId
+	 */
+	virtual ENamedThreads::Type GetRecipientThread() const = 0;
+
+	/**
+	 * Checks whether this recipient represents a local endpoint.
+	 *
+	 * Local recipients are located in the current thread or process. Recipients located in
+	 * other processes on the same machine or on remote machines are considered remote.
+	 *
+	 * @return true if this recipient is local, false otherwise.
+	 * @see IsRemote
+	 */
+	virtual bool IsLocal() const = 0;
+
+	/**
+	 * Handles the given message.
+	 *
+	 * @param Context Will hold the context of the received message.
+	 */
+	virtual void ReceiveMessage( const TSharedRef<IMessageContext, ESPMode::ThreadSafe>& Context ) = 0;
+
+public:
+
+	/**
+	 * Checks whether this recipient represents a remote endpoint.
+	 *
+	 * Local recipients are located in the current thread or process. Recipients located in
+	 * other processes on the same machine or on remote machines are considered remote.
+	 *
+	 * @return true if this recipient is remote, false otherwise.
+	 * @see IsLocal
+	 */
+	bool IsRemote() const
+	{
+		return !IsLocal();
+	}
+
+public:
+
+	/** Virtual destructor. */
+	virtual ~IReceiveMessages() { }
+};
+
+
+/** Type definition for shared pointers to instances of IReceiveMessages. */
+typedef TSharedPtr<IReceiveMessages, ESPMode::ThreadSafe> IReceiveMessagesPtr;
+
+/** Type definition for shared references to instances of IReceiveMessages. */
+typedef TSharedRef<IReceiveMessages, ESPMode::ThreadSafe> IReceiveMessagesRef;
+
+/** Type definition for shared pointers to instances of IReceiveMessages. */
+typedef TWeakPtr<IReceiveMessages, ESPMode::ThreadSafe> IReceiveMessagesWeakPtr;
diff --git a/Engine/Source/Runtime/Messaging/Public/ISendMessages.h b/Engine/Source/Runtime/Messaging/Public/ISendMessages.h
new file mode 100644
--- /dev/null
+++ b/Engine/Source/Runtime/Messaging/Public/ISendMessages.h
@@ -0,0 +1,56 @@
+// Copyright 1998-2014 Epic Games, Inc. All Rights Reserved.
+
+#pragma once
+
+
+// forward declarations
+class IMessageContext;
+
+
+/**
+ * Interface for message senders.
+ *
+ * Classes that implement this interface are able to send messages on a message bus. Each message sender must be uniquely
+ * identifiable with a message address that is returned by the @see ISendMessages.GetSenderAddress method. It is recommended
+ * that implementors of this interface generate a GUID for each instance, which must not change throughout the lifetime of
+ * the instance.
+ *
+ * The sending of messages is accomplished with the @see IMessageBus.Forward, IMessageBus.Publish and IMessageBus.Send methods.
+ * In case an error occurs during the sending of a message, the @see ISendMessages.NotifyMessageError method will be called.
+ *
+ * This interface provides a rather low-level mechanism for receiving messages. Instead of implementing it, Most users will
+ * want to use an instance of @see FMessageEndpoint, which provides a much more convenient way of sending and receiving messages.
+ *
+ * @see FMessageEndpoint, IMessageBus, IReceiveMessages
+ */
+class ISendMessages
+{
+public:
+
+	/**
+	 * Gets the sender's address.
+	 *
+	 * @return The message address.
+	 */
+	virtual FMessageAddress GetSenderAddress() = 0;
+
+	/**
+	 * Notifies the sender of errors.
+	 *
+	 * @param Context The context of the message that generated the error.
+	 * @param Error The error string.
+	 */
+	virtual void NotifyMessageError( const TSharedRef<IMessageContext, ESPMode::ThreadSafe>& Context, const FString& Error ) = 0;
+
+public:
+
+	/** Virtual destructor. */
+	virtual ~ISendMessages() { }
+};
+
+
+/** Type definition for shared pointers to instances of ISendMessages. */
+typedef TSharedPtr<ISendMessages, ESPMode::ThreadSafe> ISendMessagesPtr;
+
+/** Type definition for shared references to instances of ISendMessages. */
+typedef TSharedRef<ISendMessages, ESPMode::ThreadSafe> ISendMessagesRef;
diff --git a/Engine/Source/Runtime/Messaging/Public/Interfaces/IAuthorizeMessageRecipients.h b/Engine/Source/Runtime/Messaging/Public/Interfaces/IAuthorizeMessageRecipients.h
deleted file mode 100644
diff --git a/Engine/Source/Runtime/Messaging/Public/Interfaces/IMessageAttachment.h b/Engine/Source/Runtime/Messaging/Public/Interfaces/IMessageAttachment.h
deleted file mode 100644
diff --git a/Engine/Source/Runtime/Messaging/Public/Interfaces/IMessageBridge.h b/Engine/Source/Runtime/Messaging/Public/Interfaces/IMessageBridge.h
deleted file mode 100644
diff --git a/Engine/Source/Runtime/Messaging/Public/Interfaces/IMessageBus.h b/Engine/Source/Runtime/Messaging/Public/Interfaces/IMessageBus.h
deleted file mode 100644
diff --git a/Engine/Source/Runtime/Messaging/Public/Interfaces/IMessageContext.h b/Engine/Source/Runtime/Messaging/Public/Interfaces/IMessageContext.h
deleted file mode 100644
diff --git a/Engine/Source/Runtime/Messaging/Public/Interfaces/IMessageHandler.h b/Engine/Source/Runtime/Messaging/Public/Interfaces/IMessageHandler.h
deleted file mode 100644
diff --git a/Engine/Source/Runtime/Messaging/Public/Interfaces/IMessageInterceptor.h b/Engine/Source/Runtime/Messaging/Public/Interfaces/IMessageInterceptor.h
deleted file mode 100644
diff --git a/Engine/Source/Runtime/Messaging/Public/Interfaces/IMessageSubscription.h b/Engine/Source/Runtime/Messaging/Public/Interfaces/IMessageSubscription.h
deleted file mode 100644
diff --git a/Engine/Source/Runtime/Messaging/Public/Interfaces/IMessageTracer.h b/Engine/Source/Runtime/Messaging/Public/Interfaces/IMessageTracer.h
deleted file mode 100644
diff --git a/Engine/Source/Runtime/Messaging/Public/Interfaces/IMessageTracerBreakpoint.h b/Engine/Source/Runtime/Messaging/Public/Interfaces/IMessageTracerBreakpoint.h
deleted file mode 100644
diff --git a/Engine/Source/Runtime/Messaging/Public/Interfaces/IMessageTransport.h b/Engine/Source/Runtime/Messaging/Public/Interfaces/IMessageTransport.h
deleted file mode 100644
diff --git a/Engine/Source/Runtime/Messaging/Public/Interfaces/IMessagingModule.h b/Engine/Source/Runtime/Messaging/Public/Interfaces/IMessagingModule.h
deleted file mode 100644
diff --git a/Engine/Source/Runtime/Messaging/Public/Interfaces/IMutableMessageContext.h b/Engine/Source/Runtime/Messaging/Public/Interfaces/IMutableMessageContext.h
deleted file mode 100644
diff --git a/Engine/Source/Runtime/Messaging/Public/Interfaces/IReceiveMessages.h b/Engine/Source/Runtime/Messaging/Public/Interfaces/IReceiveMessages.h
deleted file mode 100644
diff --git a/Engine/Source/Runtime/Messaging/Public/Interfaces/ISendMessages.h b/Engine/Source/Runtime/Messaging/Public/Interfaces/ISendMessages.h
deleted file mode 100644
diff --git a/Engine/Source/Runtime/Messaging/Public/Messaging.h b/Engine/Source/Runtime/Messaging/Public/Messaging.h
deleted file mode 100644
