Messaging: Moved delegate declarations into interface classes

diff --git a/Engine/Plugins/Messaging/UdpMessaging/Source/UdpMessaging/Private/Transport/UdpMessageProcessor.h b/Engine/Plugins/Messaging/UdpMessaging/Source/UdpMessaging/Private/Transport/UdpMessageProcessor.h
--- a/Engine/Plugins/Messaging/UdpMessaging/Source/UdpMessaging/Private/Transport/UdpMessageProcessor.h
+++ b/Engine/Plugins/Messaging/UdpMessaging/Source/UdpMessaging/Private/Transport/UdpMessageProcessor.h
@@ -3,16 +3,6 @@
 #pragma once
 
 
-/**
- * Delegate type for handling the message reassembly completion.
- *
- * The first parameter is the reassembled message data.
- * The second parameter contains the optional message attachment.
- * The third parameter is the identifier of the node that sent the message.
- */
-DECLARE_DELEGATE_ThreeParams(FOnMessageProcessorMessageReassembled, const FUdpReassembledMessageRef& /*ReassembledMessage*/, const IMessageAttachmentPtr& /*Attachment*/, const FGuid& /*NodeId*/);
-
-
 /**
  * Implements a message processor for UDP messages.
  */
@@ -142,12 +132,13 @@ public:
 	 *
 	 * @param Delegate The delegate to set.
 	 */
-	FOnMessageProcessorMessageReassembled& OnMessageReassembled()
+	DECLARE_DELEGATE_ThreeParams(FOnMessageReassembled, const FUdpReassembledMessageRef& /*ReassembledMessage*/, const IMessageAttachmentPtr& /*Attachment*/, const FGuid& /*NodeId*/)
+	FOnMessageReassembled& OnMessageReassembled()
 	{
 		return MessageReassembledDelegate;
 	}
 
-	FOnMessageTransportNodeDiscovered& OnNodeDiscovered()
+	IMessageTransport::FOnNodeDiscovered& OnNodeDiscovered()
 	{
 		return NodeDiscoveredDelegate;
 	}
@@ -157,7 +148,7 @@ public:
 	 *
 	 * @param Delegate The delegate to set.
 	 */
-	FOnMessageTransportNodeLost& OnNodeLost()
+	IMessageTransport::FOnNodeLost& OnNodeLost()
 	{
 		return NodeLostDelegate;
 	}
@@ -344,13 +335,13 @@ private:
 private:
 
 	/** Holds a delegate to be invoked when a message was received on the transport channel. */
-	FOnMessageProcessorMessageReassembled MessageReassembledDelegate;
+	FOnMessageReassembled MessageReassembledDelegate;
 
 	/** Holds a delegate to be invoked when a network node was discovered. */
-	FOnMessageTransportNodeDiscovered NodeDiscoveredDelegate;
+	IMessageTransport::FOnNodeDiscovered NodeDiscoveredDelegate;
 
 	/** Holds a delegate to be invoked when a network node was lost. */
-	FOnMessageTransportNodeLost NodeLostDelegate;
+	IMessageTransport::FOnNodeLost NodeLostDelegate;
 
 private:
 
diff --git a/Engine/Plugins/Messaging/UdpMessaging/Source/UdpMessaging/Private/Transport/UdpMessageTransport.h b/Engine/Plugins/Messaging/UdpMessaging/Source/UdpMessaging/Private/Transport/UdpMessageTransport.h
--- a/Engine/Plugins/Messaging/UdpMessaging/Source/UdpMessaging/Private/Transport/UdpMessageTransport.h
+++ b/Engine/Plugins/Messaging/UdpMessaging/Source/UdpMessaging/Private/Transport/UdpMessageTransport.h
@@ -32,17 +32,17 @@ public:
 		return "UdpMessageTransport";
 	}
 
-	virtual FOnMessageTransportMessageReceived& OnMessageReceived() override
+	virtual FOnMessageReceived& OnMessageReceived() override
 	{
 		return MessageReceivedDelegate;
 	}
 
-	virtual FOnMessageTransportNodeDiscovered& OnNodeDiscovered() override
+	virtual FOnNodeDiscovered& OnNodeDiscovered() override
 	{
 		return NodeDiscoveredDelegate;
 	}
 
-	virtual FOnMessageTransportNodeLost& OnNodeLost() override
+	virtual FOnNodeLost& OnNodeLost() override
 	{
 		return NodeLostDelegate;
 	}
@@ -106,11 +106,11 @@ private:
 private:
 
 	/** Holds a delegate to be invoked when a message was received on the transport channel. */
-	FOnMessageTransportMessageReceived MessageReceivedDelegate;
+	FOnMessageReceived MessageReceivedDelegate;
 
 	/** Holds a delegate to be invoked when a network node was discovered. */
-	FOnMessageTransportNodeDiscovered NodeDiscoveredDelegate;
+	FOnNodeDiscovered NodeDiscoveredDelegate;
 
 	/** Holds a delegate to be invoked when a network node was lost. */
-	FOnMessageTransportNodeLost NodeLostDelegate;
+	FOnNodeLost NodeLostDelegate;
 };
diff --git a/Engine/Source/Runtime/Messaging/Private/Bus/MessageTracer.h b/Engine/Source/Runtime/Messaging/Private/Bus/MessageTracer.h
--- a/Engine/Source/Runtime/Messaging/Private/Bus/MessageTracer.h
+++ b/Engine/Source/Runtime/Messaging/Private/Bus/MessageTracer.h
@@ -256,17 +256,20 @@ public:
 		return Running;
 	}
 
-	virtual FMessageTracerMessageAdded& OnMessageAdded() override
+	DECLARE_DERIVED_EVENT(FMessageTracer, IMessageTracer::FOnMessageAdded, FOnMessageAdded)
+	virtual FOnMessageAdded& OnMessageAdded() override
 	{
 		return MessagesAddedDelegate;
 	}
 
-	virtual FSimpleMulticastDelegate& OnMessagesReset() override
+	DECLARE_DERIVED_EVENT(FMessageTracer, IMessageTracer::FOnMessagesReset, FOnMessagesReset)
+	virtual FOnMessagesReset& OnMessagesReset() override
 	{
 		return MessagesResetDelegate;
 	}
 
-	virtual FMessageTracerTypeAdded& OnTypeAdded() override
+	DECLARE_DERIVED_EVENT(FMessageTracer, IMessageTracer::FOnTypeAdded, FOnTypeAdded)
+	virtual FOnTypeAdded& OnTypeAdded() override
 	{
 		return TypeAddedDelegate;
 	}
@@ -462,13 +465,13 @@ private:
 private:
 
 	/** Holds a delegate that is executed when a new message has been added to the collection of known messages. */
-	FMessageTracerMessageAdded MessagesAddedDelegate;
+	FOnMessageAdded MessagesAddedDelegate;
 
 	/** Holds a delegate that is executed when the message history has been reset. */
-	FSimpleMulticastDelegate MessagesResetDelegate;
+	FOnMessagesReset MessagesResetDelegate;
 
 	/** Holds a delegate that is executed when a new type has been added to the collection of known message types. */
-	FMessageTracerTypeAdded TypeAddedDelegate;
+	FOnTypeAdded TypeAddedDelegate;
 };
 
 
diff --git a/Engine/Source/Runtime/Messaging/Public/IMessageTracer.h b/Engine/Source/Runtime/Messaging/Public/IMessageTracer.h
--- a/Engine/Source/Runtime/Messaging/Public/IMessageTracer.h
+++ b/Engine/Source/Runtime/Messaging/Public/IMessageTracer.h
@@ -152,21 +152,6 @@ struct FMessageTracerTypeInfo
 };
 
 
-/**
- * Implements a delegate that is executed when a message has been added to the collection of traced messages.
- *
- * The first parameter is the information for the message that was added.
- */
-DECLARE_MULTICAST_DELEGATE_OneParam(FMessageTracerMessageAdded, FMessageTracerMessageInfoRef)
-
-/**
- * Implements a delegate that is executed when a type has been added to the collection of traced message types.
- *
- * The first parameter is the information for the message type that was added.
- */
-DECLARE_MULTICAST_DELEGATE_OneParam(FMessageTracerTypeAdded, FMessageTracerTypeInfoRef)
-
-
 /**
  * Interface for message tracers.
  *
@@ -277,25 +262,28 @@ public:
 public:
 
 	/**
-	 * Returns a delegate that is executed when the collection of known messages has changed.
+	 * A delegate that is executed when the collection of known messages has changed.
 	 *
 	 * @return The delegate.
 	 */
-	virtual FMessageTracerMessageAdded& OnMessageAdded() = 0;
+	DECLARE_EVENT_OneParam(IMessageTracer, FOnMessageAdded, FMessageTracerMessageInfoRef)
+	virtual FOnMessageAdded& OnMessageAdded() = 0;
 
 	/**
-	 * Returns a delegate that is executed when the message history has been reset.
+	 * A delegate that is executed when the message history has been reset.
 	 *
 	 * @return The delegate.
 	 */
-	virtual FSimpleMulticastDelegate& OnMessagesReset() = 0;
+	DECLARE_EVENT(IMessageTracer, FOnMessagesReset)
+	virtual FOnMessagesReset& OnMessagesReset() = 0;
 
 	/**
-	 * Returns a delegate that is executed when the collection of known messages types has changed.
+	 * A delegate that is executed when the collection of known messages types has changed.
 	 *
 	 * @return The delegate.
 	 */
-	virtual FMessageTracerTypeAdded& OnTypeAdded() = 0;
+	DECLARE_EVENT_OneParam(IMessageTracer, FOnTypeAdded, FMessageTracerTypeInfoRef)
+	virtual FOnTypeAdded& OnTypeAdded() = 0;
 
 protected:
 
diff --git a/Engine/Source/Runtime/Messaging/Public/IMessageTransport.h b/Engine/Source/Runtime/Messaging/Public/IMessageTransport.h
--- a/Engine/Source/Runtime/Messaging/Public/IMessageTransport.h
+++ b/Engine/Source/Runtime/Messaging/Public/IMessageTransport.h
@@ -8,32 +8,6 @@ class IMessageAttachment;
 class IMessageData;
 
 
-/**
- * Delegate type for receiving transport messages.
- *
- * The first parameter is the context of the received message.
- * The second parameter is the identifier of the transport node that sent the message.
- */
-DECLARE_DELEGATE_TwoParams(FOnMessageTransportMessageReceived, const IMessageContextRef&, const FGuid&);
-
-/**
- * Delegate type for handling the discovery of transport nodes.
- *
- * The first parameter is the identifier of the transport node that was discovered.
- */
-DECLARE_DELEGATE_OneParam(FOnMessageTransportNodeDiscovered, const FGuid&);
-
-/**
- * Delegate type for handling disconnected or timed out transport nodes.
- *
- * Remote nodes can be lost if they disconnect from the bus or if a keep-alive
- * message was not received within the required timeout (i.e. in case of a crash).
- *
- * The first parameter is the identifier of the transport node that was lost.
- */
-DECLARE_DELEGATE_OneParam(FOnMessageTransportNodeLost, const FGuid&);
-
-
 /**
  * Interface for message transport technologies.
  *
@@ -78,25 +52,28 @@ public:
 public:
 
 	/**
-	 * Returns a delegate that is executed when message data has been received.
+	 * A delegate that is executed when message data has been received.
 	 *
 	 * @param Delegate The delegate to set.
 	 */
-	virtual FOnMessageTransportMessageReceived& OnMessageReceived() = 0;
+	DECLARE_DELEGATE_TwoParams(FOnMessageReceived, const IMessageContextRef&, const FGuid&)
+	virtual FOnMessageReceived& OnMessageReceived() = 0;
 
 	/**
-	 * Returns a delegate that is executed when a transport node has been discovered.
+	 * A delegate that is executed when a transport node has been discovered.
 	 *
 	 * @param Delegate The delegate to set.
 	 */
-	virtual FOnMessageTransportNodeDiscovered& OnNodeDiscovered() = 0;
+	DECLARE_DELEGATE_OneParam(FOnNodeDiscovered, const FGuid&)
+	virtual FOnNodeDiscovered& OnNodeDiscovered() = 0;
 
 	/**
-	 * Returns a delegate that is executed when a transport node has closed or timed out.
+	 * A delegate that is executed when a transport node has closed or timed out.
 	 *
 	 * @param Delegate The delegate to set.
 	 */
-	virtual FOnMessageTransportNodeLost& OnNodeLost() = 0;
+	DECLARE_DELEGATE_OneParam(FOnNodeLost, const FGuid&)
+	virtual FOnNodeLost& OnNodeLost() = 0;
 
 protected:
 
