Messaging: Fixed broken online documentation; removed redundant comments

diff --git a/Engine/Source/Runtime/Messaging/Public/IAuthorizeMessageRecipients.h b/Engine/Source/Runtime/Messaging/Public/IAuthorizeMessageRecipients.h
--- a/Engine/Source/Runtime/Messaging/Public/IAuthorizeMessageRecipients.h
+++ b/Engine/Source/Runtime/Messaging/Public/IAuthorizeMessageRecipients.h
@@ -3,7 +3,6 @@
 #pragma once
 
 
-// forward declarations
 struct FMessageAddress;
 class IMessageInterceptor;
 class IReceiveMessages;
diff --git a/Engine/Source/Runtime/Messaging/Public/IMessageBridge.h b/Engine/Source/Runtime/Messaging/Public/IMessageBridge.h
--- a/Engine/Source/Runtime/Messaging/Public/IMessageBridge.h
+++ b/Engine/Source/Runtime/Messaging/Public/IMessageBridge.h
@@ -15,7 +15,7 @@
  * separate processes or en different computers. Another common use case is to connect an Unreal Engine
  * message bus to an entirely different messaging system that is not based on Unreal Engine.
  *
- * Message bridge instances can be created with the @see IMessagingModule.CreateBridge method.
+ * Message bridge instances can be created with the IMessagingModule.CreateBridge method.
  */
 class IMessageBridge
 {
diff --git a/Engine/Source/Runtime/Messaging/Public/IMessageBus.h b/Engine/Source/Runtime/Messaging/Public/IMessageBus.h
--- a/Engine/Source/Runtime/Messaging/Public/IMessageBus.h
+++ b/Engine/Source/Runtime/Messaging/Public/IMessageBus.h
@@ -3,7 +3,6 @@
 #pragma once
 
 
-// forward declarations
 enum class EMessageScope;
 struct FMessageAddress;
 class IMessageAttachment;
@@ -30,7 +29,7 @@ DECLARE_MULTICAST_DELEGATE(FOnMessageBusShutdown);
  * Depending on their usage, messages are classified into a number of types, such as commands, events and documents. In Unreal Engine,
  * all these messages are modeled with regular built-in or user defined UStructs that may either be empty or contain data in the
  * form of UProperty fields[2]. Before being dispatched, messages are internally wrapped into so called Message Context objects
- * (@see IMessageContext) that contain additional information about a message, such as when it was sent and the sender and recipients.
+ * (see IMessageContext) that contain additional information about a message, such as when it was sent and the sender and recipients.
  *
  * The sending and receiving of messages is not limited to message endpoints within the same thread or process, but may be extended
  * to other applications running on the same computer or on other computers connected to a network with so called Message Transport
@@ -38,13 +37,13 @@ DECLARE_MULTICAST_DELEGATE(FOnMessageBusShutdown);
  * details of the underlying transport mechanisms, so that users can focus on implementing their distributed applications rather than
  * worrying about how data gets from one endpoint to another. A message bus makes it look as if all senders and recipients are located
  * within the same process, regardless of whether that is actually the case or not. It is possible to restrict the reach of messages
- * using so called Message Scopes (@see EMessageScope).
+ * using so called Message Scopes (see EMessageScope).
  *
- * All message recipients (objects implementing the @see IReceiveMessages interface) must be registered with the message bus through the
- * @see IMessageBus.Register method. Before a recipient is destroyed it should unregister itself using the @see IMessageBus.Unregister
- * method. Message senders (objects implementing the @see ISendMessages interface) do not register with the bus, but instead pass a
+ * All message recipients (objects implementing the see IReceiveMessages interface) must be registered with the message bus through the
+ * see IMessageBus.Register method. Before a recipient is destroyed it should unregister itself using the see IMessageBus.Unregister
+ * method. Message senders (objects implementing the see ISendMessages interface) do not register with the bus, but instead pass a
  * reference to themselves each time they send a message. The IReceiveMessages and ISendMessages both are very low-level interfaces
- * into the messaging system. Most users will prefer to use instances of the @see FMessageEndpoint class instead, which provides a much
+ * into the messaging system. Most users will prefer to use instances of the see FMessageEndpoint class instead, which provides a much
  * more convenient way of sending and receiving messages.
  *
  * Message buses in Unreal Engine support the following two common messaging patterns: Request-Reply and Publish-Subscribe. Please note
@@ -52,12 +51,12 @@ DECLARE_MULTICAST_DELEGATE(FOnMessageBusShutdown);
  * separate features in the future.
  *
  * In the Request-Reply pattern a message is sent to one or more particular message recipients using the IMessageBus.Send method.
- * Message recipients implement the @see IMessageRecipient interface and are uniquely identified by their addresses (@see FMessageAddress).
+ * Message recipients implement the see IMessageRecipient interface and are uniquely identified by their addresses (see FMessageAddress).
  * After a message is received, the recipients may then reply with another message using the same IMessageBus.Send method. Alternatively,
- * a previously received message may be forwarded to another recipient using the @see IMessageBus.Forward method. This pattern is useful
+ * a previously received message may be forwarded to another recipient using the see IMessageBus.Forward method. This pattern is useful
  * when message recipients already know about each other and wish to communicate directly, i.e. to exchange commands or events.
  *
- * In the Publish-Subscribe pattern a message is sent to all message recipients on the bus using the @see IMessageBus.Publish method.
+ * In the Publish-Subscribe pattern a message is sent to all message recipients on the bus using the see IMessageBus.Publish method.
  * Only recipients that previously subscribed to the type of the sent message using the @IMessageBus.Subscribe method will actually receive
  * the message. All other recipients will not receive the message. After a published message is received, the recipients may respond with
  * another message, either directly to the message sender using IMessageBus.Send method, or by publishing another message using the
@@ -80,7 +79,7 @@ DECLARE_MULTICAST_DELEGATE(FOnMessageBusShutdown);
  * manipulation of messages contents, such as message filtering and enriching, splitting and aggregating, resequencing or authentication[4].
  *
  * The messaging system also provides a facility for debugging the system itself through the so called Message Tracer (an internal object
- * implementing the @see IMessageTracer interface) that can be accessed with the IMessageBus.GetTracer() method. The message tracer is
+ * implementing the see IMessageTracer interface) that can be accessed with the IMessageBus.GetTracer() method. The message tracer is
  * currently used in the Messaging Debugger - a visual debugging tool for the Messaging System in Unreal Frontend and the Unreal Editor.
  *
  * Notes:
diff --git a/Engine/Source/Runtime/Messaging/Public/IMessageContext.h b/Engine/Source/Runtime/Messaging/Public/IMessageContext.h
--- a/Engine/Source/Runtime/Messaging/Public/IMessageContext.h
+++ b/Engine/Source/Runtime/Messaging/Public/IMessageContext.h
@@ -5,7 +5,6 @@
 #include "TaskGraphInterfaces.h"
 
 
-// forward declarations
 class IMessageAttachment;
 class UScriptStruct;
 
diff --git a/Engine/Source/Runtime/Messaging/Public/IMessageHandler.h b/Engine/Source/Runtime/Messaging/Public/IMessageHandler.h
--- a/Engine/Source/Runtime/Messaging/Public/IMessageHandler.h
+++ b/Engine/Source/Runtime/Messaging/Public/IMessageHandler.h
@@ -5,7 +5,6 @@
 #include "CoreUObject.h"
 
 
-// forward declarations
 class IMessageContext;
 
 
diff --git a/Engine/Source/Runtime/Messaging/Public/IMessageInterceptor.h b/Engine/Source/Runtime/Messaging/Public/IMessageInterceptor.h
--- a/Engine/Source/Runtime/Messaging/Public/IMessageInterceptor.h
+++ b/Engine/Source/Runtime/Messaging/Public/IMessageInterceptor.h
@@ -3,7 +3,6 @@
 #pragma once
 
 
-// forward declarations
 class IMessageContext;
 
 
diff --git a/Engine/Source/Runtime/Messaging/Public/IMessageSubscription.h b/Engine/Source/Runtime/Messaging/Public/IMessageSubscription.h
--- a/Engine/Source/Runtime/Messaging/Public/IMessageSubscription.h
+++ b/Engine/Source/Runtime/Messaging/Public/IMessageSubscription.h
@@ -3,7 +3,6 @@
 #pragma once
 
 
-// forward declarations
 enum class EMessageScope;
 class IReceiveMessages;
 
@@ -12,8 +11,8 @@ class IReceiveMessages;
  * Interface for message subscriptions.
  *
  * This interface provides access to a message subscription that was previously created with the
- * @see IMessageBus.Subscribe method. It can be used to query the subscription's details using the
- * various getters and modify its enabled state using the @see IMessageSubscription.Enable and
+ * IMessageBus.Subscribe method. It can be used to query the subscription's details using the
+ * various getters and modify its enabled state using the IMessageSubscription.Enable and
  * IMessageSubscription.Disable methods.
  *
  * A subscription that is disabled will cause a message endpoint to stop receiving messages for
diff --git a/Engine/Source/Runtime/Messaging/Public/IMessageTracer.h b/Engine/Source/Runtime/Messaging/Public/IMessageTracer.h
--- a/Engine/Source/Runtime/Messaging/Public/IMessageTracer.h
+++ b/Engine/Source/Runtime/Messaging/Public/IMessageTracer.h
@@ -3,7 +3,6 @@
 #pragma once
 
 
-// forward declarations
 class IMessageContext;
 
 
diff --git a/Engine/Source/Runtime/Messaging/Public/IMessageTracerBreakpoint.h b/Engine/Source/Runtime/Messaging/Public/IMessageTracerBreakpoint.h
--- a/Engine/Source/Runtime/Messaging/Public/IMessageTracerBreakpoint.h
+++ b/Engine/Source/Runtime/Messaging/Public/IMessageTracerBreakpoint.h
@@ -3,7 +3,6 @@
 #pragma once
 
 
-// forward declarations
 class IMessageContext;
 
 
diff --git a/Engine/Source/Runtime/Messaging/Public/IMessageTransport.h b/Engine/Source/Runtime/Messaging/Public/IMessageTransport.h
--- a/Engine/Source/Runtime/Messaging/Public/IMessageTransport.h
+++ b/Engine/Source/Runtime/Messaging/Public/IMessageTransport.h
@@ -3,7 +3,6 @@
 #pragma once
 
 
-// forward declarations
 class IMessageAttachment;
 class IMessageData;
 
diff --git a/Engine/Source/Runtime/Messaging/Public/IMessagingModule.h b/Engine/Source/Runtime/Messaging/Public/IMessagingModule.h
--- a/Engine/Source/Runtime/Messaging/Public/IMessagingModule.h
+++ b/Engine/Source/Runtime/Messaging/Public/IMessagingModule.h
@@ -5,7 +5,6 @@
 #include "ModuleInterface.h"
 
 
-// forward declarations
 struct FMessageAddress;
 class IAuthorizeMessageRecipients;
 class IMessageBus;
diff --git a/Engine/Source/Runtime/Messaging/Public/IMutableMessageContext.h b/Engine/Source/Runtime/Messaging/Public/IMutableMessageContext.h
--- a/Engine/Source/Runtime/Messaging/Public/IMutableMessageContext.h
+++ b/Engine/Source/Runtime/Messaging/Public/IMutableMessageContext.h
@@ -5,7 +5,6 @@
 #include "IMessageContext.h"
 
 
-// forward declarations
 class IMessageAttachment;
 class UScriptStruct;
 
diff --git a/Engine/Source/Runtime/Messaging/Public/IReceiveMessages.h b/Engine/Source/Runtime/Messaging/Public/IReceiveMessages.h
--- a/Engine/Source/Runtime/Messaging/Public/IReceiveMessages.h
+++ b/Engine/Source/Runtime/Messaging/Public/IReceiveMessages.h
@@ -3,7 +3,6 @@
 #pragma once
 
 
-// forward declarations
 class IMessageContext;
 
 
@@ -11,11 +10,11 @@ class IMessageContext;
  * Interface for message recipients.
  *
  * Classes that implement this interface are able to receive messages from a message bus. A message recipient will receive
- * a call to its ReceiveMessage() method for each message that was sent directly to it (@see IMessageBus.Send) and for each
- * published message that it subscribed to (@see IMessageBus.Publish).
+ * a call to its ReceiveMessage() method for each message that was sent directly to it (see IMessageBus.Send) and for each
+ * published message that it subscribed to (see IMessageBus.Publish).
  *
  * This interface provides a rather low-level mechanism for receiving messages. Instead of implementing it, Most users will
- * want to use an instance of @see FMessageEndpoint, which provides a much more convenient way of sending and receiving messages.
+ * want to use an instance of see FMessageEndpoint, which provides a much more convenient way of sending and receiving messages.
  *
  * @see FMessageEndpoint, IMessageBus, ISendMessages
  */
diff --git a/Engine/Source/Runtime/Messaging/Public/ISendMessages.h b/Engine/Source/Runtime/Messaging/Public/ISendMessages.h
--- a/Engine/Source/Runtime/Messaging/Public/ISendMessages.h
+++ b/Engine/Source/Runtime/Messaging/Public/ISendMessages.h
@@ -3,7 +3,6 @@
 #pragma once
 
 
-// forward declarations
 class IMessageContext;
 
 
@@ -11,15 +10,15 @@ class IMessageContext;
  * Interface for message senders.
  *
  * Classes that implement this interface are able to send messages on a message bus. Each message sender must be uniquely
- * identifiable with a message address that is returned by the @see ISendMessages.GetSenderAddress method. It is recommended
+ * identifiable with a message address that is returned by the ISendMessages.GetSenderAddress method. It is recommended
  * that implementors of this interface generate a GUID for each instance, which must not change throughout the lifetime of
  * the instance.
  *
- * The sending of messages is accomplished with the @see IMessageBus.Forward, IMessageBus.Publish and IMessageBus.Send methods.
- * In case an error occurs during the sending of a message, the @see ISendMessages.NotifyMessageError method will be called.
+ * The sending of messages is accomplished with the IMessageBus.Forward, IMessageBus.Publish and IMessageBus.Send methods.
+ * In case an error occurs during the sending of a message, the ISendMessages.NotifyMessageError method will be called.
  *
  * This interface provides a rather low-level mechanism for receiving messages. Instead of implementing it, Most users will
- * want to use an instance of @see FMessageEndpoint, which provides a much more convenient way of sending and receiving messages.
+ * want to use an instance of FMessageEndpoint, which provides a much more convenient way of sending and receiving messages.
  *
  * @see FMessageEndpoint, IMessageBus, IReceiveMessages
  */
