Messaging: Allowing message handlers to decide for themselves whether they want to handle a message; documentation updates.

#UpgradeNotes:
If one of your classes implements IMessageHandler, you must remove GetHandledMessageType() and instead verify the type of message in HandleMessage(). See TMessageHandler<> for details.

diff --git a/Engine/Source/Runtime/Messaging/Public/Helpers/MessageEndpoint.h b/Engine/Source/Runtime/Messaging/Public/Helpers/MessageEndpoint.h
--- a/Engine/Source/Runtime/Messaging/Public/Helpers/MessageEndpoint.h
+++ b/Engine/Source/Runtime/Messaging/Public/Helpers/MessageEndpoint.h
@@ -781,12 +781,7 @@ protected:
 
 		for (int32 HandlerIndex = 0; HandlerIndex < Handlers.Num(); ++HandlerIndex)
 		{
-			IMessageHandlerPtr& Handler = Handlers[HandlerIndex];
-
-			if (Handler->GetHandledMessageType() == Context->GetMessageType())
-			{
-				Handler->HandleMessage(Context);
-			}
+			Handlers[HandlerIndex]->HandleMessage(Context);
 		}
 	}
 
diff --git a/Engine/Source/Runtime/Messaging/Public/Helpers/MessageEndpointBuilder.h b/Engine/Source/Runtime/Messaging/Public/Helpers/MessageEndpointBuilder.h
--- a/Engine/Source/Runtime/Messaging/Public/Helpers/MessageEndpointBuilder.h
+++ b/Engine/Source/Runtime/Messaging/Public/Helpers/MessageEndpointBuilder.h
@@ -51,12 +51,16 @@ public:
 	/**
 	 * Adds a message handler for the given type of messages.
 	 *
+	 * It is legal to configure multiple handlers for the same message type. Each
+	 * handler will be executed when a message of the specified type is received.
+	 *
 	 * @param HandlerType The type of the object handling the messages.
 	 * @param MessageType The type of messages to handle.
 	 * @param Handler The class handling the messages.
 	 * @param HandlerFunc The class function handling the messages.
 	 * @param ReceivingThread The thread on which to handle the message.
 	 * @return This instance (for method chaining).
+	 * @see WithHandler
 	 */
 	template<typename MessageType, typename HandlerType>
 	FMessageEndpointBuilder& Handling( HandlerType* Handler, typename TMessageHandlerFunc<MessageType, HandlerType>::Type HandlerFunc )
@@ -125,8 +129,12 @@ public:
 	/**
 	 * Registers a message handler with the endpoint.
 	 *
+	 * It is legal to configure multiple handlers for the same message type. Each
+	 * handler will be executed when a message of the specified type is received.
+	 *
 	 * @param Handler The handler to add.
 	 * @return This instance (for method chaining).
+	 * @see Handling
 	 */
 	FMessageEndpointBuilder& WithHandler( const IMessageHandlerRef& Handler )
 	{
diff --git a/Engine/Source/Runtime/Messaging/Public/IMessageHandler.h b/Engine/Source/Runtime/Messaging/Public/IMessageHandler.h
--- a/Engine/Source/Runtime/Messaging/Public/IMessageHandler.h
+++ b/Engine/Source/Runtime/Messaging/Public/IMessageHandler.h
@@ -15,19 +15,10 @@ class IMessageHandler
 {
 public:
 
-	/**
-	 * Gets the message type handled by this handler.
-	 *
-	 * @return The name of the message type.
-	 * @see HandleMessage
-	 */
-	virtual const FName GetHandledMessageType() const = 0;
-
 	/**
 	 * Handles the specified message.
 	 *
 	 * @param Context The context of the message to handle.
-	 * @see GetHandledMessageType
 	 */
 	virtual void HandleMessage( const TSharedRef<IMessageContext, ESPMode::ThreadSafe>& Context ) = 0;
 
@@ -59,7 +50,10 @@ struct TMessageHandlerFunc
 
 
 /**
- * Template for message handlers.
+ * Template for handlers of one specific message type.
+ *
+ * This class is used for message handlers that handle one particular type of messages.
+ * For handlers that can handle any type of messages see FMessageHandler.
  *
  * @param MessageType The type of message to handle.
  * @param HandlerType The type of the handler class.
@@ -83,18 +77,19 @@ public:
 		check(InHandler != nullptr);
 	}
 
+	/** Virtual destructor. */
+	~TMessageHandler() { }
+
 public:
 
 	// IMessageHandler interface
 	
-	virtual const FName GetHandledMessageType() const override
-	{
-		return MessageType::StaticStruct()->GetFName();
-	}
-
 	virtual void HandleMessage( const TSharedRef<IMessageContext, ESPMode::ThreadSafe>& Context ) override
 	{
-		(Handler->*HandlerFunc)(*static_cast<const MessageType*>(Context->GetMessage()), Context);
+		if (Context->GetMessageType() == MessageType::StaticStruct()->GetFName())
+		{
+			(Handler->*HandlerFunc)(*static_cast<const MessageType*>(Context->GetMessage()), Context);
+		}	
 	}
 	
 private:
