Messaging: Fix for message dispatch sometimes taking a very long time

diff --git a/Engine/Source/Runtime/Messaging/Public/Helpers/MessageEndpoint.h b/Engine/Source/Runtime/Messaging/Public/Helpers/MessageEndpoint.h
--- a/Engine/Source/Runtime/Messaging/Public/Helpers/MessageEndpoint.h
+++ b/Engine/Source/Runtime/Messaging/Public/Helpers/MessageEndpoint.h
@@ -163,7 +163,18 @@ public:
 	 */
 	void SetRecipientThread( const ENamedThreads::Type& NamedThread )
 	{
-		RecipientThread = NamedThread;
+		if (NamedThread == ENamedThreads::GameThread_Local)
+		{
+			RecipientThread = ENamedThreads::GameThread;
+		}
+		else if (NamedThread == ENamedThreads::RenderThread_Local)
+		{
+			RecipientThread = ENamedThreads::RenderThread;
+		}
+		else
+		{
+			RecipientThread = NamedThread;
+		}		
 	}
 
 public:
