Messaging: Added support for lambdas and function objects in message endpoint handlers

diff --git a/Engine/Source/Runtime/Messaging/Public/Helpers/MessageEndpointBuilder.h b/Engine/Source/Runtime/Messaging/Public/Helpers/MessageEndpointBuilder.h
--- a/Engine/Source/Runtime/Messaging/Public/Helpers/MessageEndpointBuilder.h
+++ b/Engine/Source/Runtime/Messaging/Public/Helpers/MessageEndpointBuilder.h
@@ -49,26 +49,50 @@ public:
 public:
 
 	/**
-	 * Adds a message handler for the given type of messages.
+	 * Adds a message handler for the given type of messages (via raw function pointers).
 	 *
 	 * It is legal to configure multiple handlers for the same message type. Each
 	 * handler will be executed when a message of the specified type is received.
 	 *
+	 * This overload is used to register raw class member functions.
+	 *
 	 * @param HandlerType The type of the object handling the messages.
 	 * @param MessageType The type of messages to handle.
 	 * @param Handler The class handling the messages.
 	 * @param HandlerFunc The class function handling the messages.
-	 * @param ReceivingThread The thread on which to handle the message.
 	 * @return This instance (for method chaining).
 	 * @see WithHandler
 	 */
 	template<typename MessageType, typename HandlerType>
-	FMessageEndpointBuilder& Handling( HandlerType* Handler, typename TMessageHandlerFunc<MessageType, HandlerType>::Type HandlerFunc )
+	FMessageEndpointBuilder& Handling( HandlerType* Handler, typename TRawMessageHandler<MessageType, HandlerType>::FuncType HandlerFunc )
 	{
 		// @todo gmp: implement proper async message deserialization, so this can be removed
 		static_assert(TStructOpsTypeTraits<MessageType>::WithMessageHandling == true, "Please add a WithMessageHandling type trait.");
+		Handlers.Add(MakeShareable(new TRawMessageHandler<MessageType, HandlerType>(Handler, HandlerFunc)));
+
+		return *this;
+	}
 
-		Handlers.Add(MakeShareable(new TMessageHandler<MessageType, HandlerType>(Handler, HandlerFunc)));
+	/**
+	 * Adds a message handler for the given type of messages.
+	 *
+	 * It is legal to configure multiple handlers for the same message type. Each
+	 * handler will be executed when a message of the specified type is received.
+	 *
+	 * This overload is used to register functions that are compatible with TFunction
+	 * function objects, such as global and static functions, as well as lambdas.
+	 *
+	 * @param MessageType The type of messages to handle.
+	 * @param Function The function object handling the messages.
+	 * @return This instance (for method chaining).
+	 * @see WithHandler
+	 */
+	template<typename MessageType>
+	FMessageEndpointBuilder& Handling( typename TFunctionMessageHandler<MessageType>::FuncType HandlerFunc )
+	{
+		// @todo gmp: implement proper async message deserialization, so this can be removed
+		static_assert(TStructOpsTypeTraits<MessageType>::WithMessageHandling == true, "Please add a WithMessageHandling type trait.");
+		Handlers.Add(MakeShareable(new TFunctionMessageHandler<MessageType>(HandlerFunc)));
 
 		return *this;
 	}
diff --git a/Engine/Source/Runtime/Messaging/Public/IMessageHandler.h b/Engine/Source/Runtime/Messaging/Public/IMessageHandler.h
--- a/Engine/Source/Runtime/Messaging/Public/IMessageHandler.h
+++ b/Engine/Source/Runtime/Messaging/Public/IMessageHandler.h
@@ -37,31 +37,20 @@ typedef TSharedRef<IMessageHandler, ESPMode::ThreadSafe> IMessageHandlerRef;
 
 
 /**
- * Template for message handler function pointer types.
+ * Template for handlers of one specific message type (via raw function pointers).
  *
  * @param MessageType The type of message to handle.
  * @param HandlerType The type of the handler class.
  */
 template<typename MessageType, typename HandlerType>
-struct TMessageHandlerFunc
+class TRawMessageHandler
+	: public IMessageHandler
 {
-	typedef void (HandlerType::*Type)( const MessageType&, const TSharedRef<IMessageContext, ESPMode::ThreadSafe>& );
-};
+public:
 
+	/** Type definition for function pointers that are compatible with this TRawMessageHandler. */
+	typedef void (HandlerType::*FuncType)( const MessageType&, const TSharedRef<IMessageContext, ESPMode::ThreadSafe>& );
 
-/**
- * Template for handlers of one specific message type.
- *
- * This class is used for message handlers that handle one particular type of messages.
- * For handlers that can handle any type of messages see FMessageHandler.
- *
- * @param MessageType The type of message to handle.
- * @param HandlerType The type of the handler class.
- */
-template<typename MessageType, typename HandlerType>
-class TMessageHandler
-	: public IMessageHandler
-{
 public:
 
 	/**
@@ -70,7 +59,7 @@ public:
 	 * @param InHandler The object that will handle the messages.
 	 * @param InHandlerFunc The object's message handling function.
 	 */
-	TMessageHandler( HandlerType* InHandler, typename TMessageHandlerFunc<MessageType, HandlerType>::Type InHandlerFunc )
+	TRawMessageHandler( HandlerType* InHandler, FuncType InHandlerFunc )
 		: Handler(InHandler)
 		, HandlerFunc(InHandlerFunc)
 	{
@@ -78,7 +67,7 @@ public:
 	}
 
 	/** Virtual destructor. */
-	~TMessageHandler() { }
+	~TRawMessageHandler() { }
 
 public:
 
@@ -98,5 +87,53 @@ private:
 	HandlerType* Handler;
 
 	/** Holds a pointer to the actual handler function. */
-	typename TMessageHandlerFunc<MessageType, HandlerType>::Type HandlerFunc;
+	FuncType HandlerFunc;
+};
+
+
+/**
+ * Template for handlers of one specific message type (via function objects).
+ *
+ * @param MessageType The type of message to handle.
+ */
+template<typename MessageType>
+class TFunctionMessageHandler
+	: public IMessageHandler
+{
+public:
+
+	/** Type definition for function objects that are compatible with this TFunctionMessageHandler. */
+	typedef TFunction<void(const MessageType&, const TSharedRef<IMessageContext, ESPMode::ThreadSafe>&)> FuncType;
+
+public:
+
+	/**
+	 * Creates and initializes a new message handler.
+	 *
+	 * @param InHandler The object that will handle the messages.
+	 * @param InHandlerFunc The object's message handling function.
+	 */
+	TFunctionMessageHandler( FuncType InHandlerFunc )
+		: HandlerFunc(InHandlerFunc)
+	{ }
+
+	/** Virtual destructor. */
+	~TFunctionMessageHandler() { }
+
+public:
+
+	// IMessageHandler interface
+	
+	virtual void HandleMessage( const TSharedRef<IMessageContext, ESPMode::ThreadSafe>& Context ) override
+	{
+		if (Context->GetMessageType() == MessageType::StaticStruct()->GetFName())
+		{
+			HandlerFunc(*static_cast<const MessageType*>(Context->GetMessage()), Context);
+		}	
+	}
+	
+private:
+
+	/** Holds a pointer to the actual handler function. */
+	FuncType HandlerFunc;
 };
