Messaging: Refactored message tracer to use lambdas instead of delegates

diff --git a/Engine/Source/Runtime/Messaging/Private/Bus/MessageTracer.cpp b/Engine/Source/Runtime/Messaging/Private/Bus/MessageTracer.cpp
--- a/Engine/Source/Runtime/Messaging/Private/Bus/MessageTracer.cpp
+++ b/Engine/Source/Runtime/Messaging/Private/Bus/MessageTracer.cpp
@@ -22,238 +22,367 @@ FMessageTracer::~FMessageTracer()
 }
 
 
-/* IMessageTracer interface
+/* FMessageTracer interface
  *****************************************************************************/
 
-int32 FMessageTracer::GetEndpoints( TArray<FMessageTracerEndpointInfoPtr>& OutEndpoints ) const
+void FMessageTracer::TraceAddedInterceptor( const IMessageInterceptorRef& Interceptor, const FName& MessageType )
 {
-	RecipientsToEndpointInfos.GenerateValueArray(OutEndpoints);
+	double Timestamp = FPlatformTime::Seconds();
 
-	return OutEndpoints.Num();
+	Traces.Enqueue(
+		[=]() {
+			// @todo gmp: trace added message interceptors
+		}
+	);
 }
 
 
-int32 FMessageTracer::GetMessages( TArray<FMessageTracerMessageInfoPtr>& OutMessages ) const
+void FMessageTracer::TraceAddedRecipient( const FMessageAddress& Address, const IReceiveMessagesRef& Recipient )
 {
-	MessageInfos.GenerateValueArray(OutMessages);
-
-	return OutMessages.Num();
+	double Timestamp = FPlatformTime::Seconds();
+
+	Traces.Enqueue(
+		[=]() {
+			// create endpoint information
+			FMessageTracerEndpointInfoPtr& EndpointInfo = RecipientsToEndpointInfos.FindOrAdd(Recipient->GetRecipientId());
+
+			if (!EndpointInfo.IsValid())
+			{
+				EndpointInfo = MakeShareable(new FMessageTracerEndpointInfo());
+			}
+
+			// initialize endpoint information
+			FMessageTracerAddressInfoRef AddressInfo = MakeShareable(new FMessageTracerAddressInfo());
+			AddressInfo->Address = Address;
+			AddressInfo->TimeRegistered = Timestamp;
+			AddressInfo->TimeUnregistered = 0;
+
+			EndpointInfo->AddressInfos.Add(Address, AddressInfo);
+			EndpointInfo->Name = Recipient->GetDebugName();
+			EndpointInfo->Remote = Recipient->IsRemote();
+
+			// add to address table
+			AddressesToEndpointInfos.Add(Address, EndpointInfo);
+		}
+	);
 }
 
 
-int32 FMessageTracer::GetMessageTypes( TArray<FMessageTracerTypeInfoPtr>& OutTypes ) const
+void FMessageTracer::TraceAddedSubscription( const IMessageSubscriptionRef& Subscription )
 {
-	MessageTypes.GenerateValueArray(OutTypes);
+	if (!Running)
+	{
+		return;
+	}
 
-	return OutTypes.Num();
+	double Timestamp = FPlatformTime::Seconds();
+
+	Traces.Enqueue(
+		[=]() {
+			// @todo gmp: trace added subscriptions
+		}
+	);
 }
 
 
-bool FMessageTracer::Tick( float DeltaTime )
+void FMessageTracer::TraceDispatchedMessage( const IMessageContextRef& Context, const IReceiveMessagesRef& Recipient, bool Async )
 {
-	if (ResetPending)
+	if (!Running)
 	{
-		ResetMessages();
-		ResetPending = false;
+		return;
 	}
 
-	if (Traces.IsEmpty())
-	{
-		return false;
-	}
+	double Timestamp = FPlatformTime::Seconds();
 
-	// process new traces
-	TraceDelegate Trace;
+	Traces.Enqueue(
+		[=]() {
+			// look up message & endpoint info
+			FMessageTracerMessageInfoPtr MessageInfo = MessageInfos.FindRef(Context);
 
-	while (Traces.Dequeue(Trace))
-	{
-		Trace.Execute();
-	}
+			if (!MessageInfo.IsValid())
+			{
+				return;
+			}
 
-	return true;
-}
+			FMessageTracerEndpointInfoPtr& EndpointInfo = RecipientsToEndpointInfos.FindOrAdd(Recipient->GetRecipientId());
 
+			if (!EndpointInfo.IsValid())
+			{
+				return;
+			}
 
-/* FMessageTracer implementation
- *****************************************************************************/
+			// update message information
+			FMessageTracerDispatchStateRef DispatchState = MakeShareable(new FMessageTracerDispatchState());
 
-void FMessageTracer::ProcessAddedInterceptor( FString Name, FName MessageType, double TimeSeconds )
-{
-	
+			DispatchState->DispatchLatency = Timestamp - MessageInfo->TimeSent;
+			DispatchState->DispatchType = Async ? EMessageTracerDispatchTypes::TaskGraph : EMessageTracerDispatchTypes::Direct;
+			DispatchState->EndpointInfo = EndpointInfo;
+			DispatchState->TimeDispatched = Timestamp;
+			DispatchState->TimeHandled = 0.0;
+
+			MessageInfo->DispatchStates.Add(EndpointInfo, DispatchState);
+
+			// update database
+			EndpointInfo->ReceivedMessages.Add(MessageInfo);
+		}
+	);
 }
 
 
-void FMessageTracer::ProcessAddedRecipient( FMessageAddress Address, FRecipientInfo RecipientInfo, double TimeSeconds )
+void FMessageTracer::TraceHandledMessage( const IMessageContextRef& Context, const IReceiveMessagesRef& Recipient )
 {
-	// create endpoint information
-	FMessageTracerEndpointInfoPtr& EndpointInfo = RecipientsToEndpointInfos.FindOrAdd(RecipientInfo.Id);
-
-	if (!EndpointInfo.IsValid())
+	if (!Running)
 	{
-		EndpointInfo = MakeShareable(new FMessageTracerEndpointInfo());
+		return;
 	}
 
-	// initialize endpoint information
-	FMessageTracerAddressInfoRef AddressInfo = MakeShareable(new FMessageTracerAddressInfo());
-	AddressInfo->Address = Address;
-	AddressInfo->TimeRegistered = TimeSeconds;
-	AddressInfo->TimeUnregistered = 0;
+	double Timestamp = FPlatformTime::Seconds();
+
+	Traces.Enqueue(
+		[=]() {
+			// look up message & endpoint info
+			FMessageTracerMessageInfoPtr MessageInfo = MessageInfos.FindRef(Context);
+
+			if (!MessageInfo.IsValid())
+			{
+				return;
+			}
+
+			FMessageTracerEndpointInfoPtr EndpointInfo = RecipientsToEndpointInfos.FindRef(Recipient->GetRecipientId());
 
-	EndpointInfo->AddressInfos.Add(Address, AddressInfo);
-	EndpointInfo->Name = RecipientInfo.Name;
-	EndpointInfo->Remote = RecipientInfo.Remote;
+			if (!EndpointInfo.IsValid())
+			{
+				return;
+			}
 
-	// add to address table
-	AddressesToEndpointInfos.Add(Address, EndpointInfo);
+			// update message information
+			FMessageTracerDispatchStatePtr DispatchState = MessageInfo->DispatchStates.FindRef(EndpointInfo);
+
+			if (DispatchState.IsValid())
+			{
+				DispatchState->TimeHandled = Timestamp;
+			}
+		}
+	);
 }
 
 
-void FMessageTracer::ProcessAddedSubscriptionTrace( double TimeSeconds )
+void FMessageTracer::TraceInterceptedMessage( const IMessageContextRef& Context, const IMessageInterceptorRef& Interceptor )
 {
+	if (!Running)
+	{
+		return;
+	}
+
+	double Timestamp = FPlatformTime::Seconds();
 
+	Traces.Enqueue(
+		[=]() {
+			// @todo gmp: trace intercepted messages
+		}
+	);		
 }
 
 
-void FMessageTracer::ProcessDispatchedMessage( IMessageContextRef Context, double TimeSeconds, FGuid RecipientId, bool Async )
+void FMessageTracer::TraceRemovedInterceptor( const IMessageInterceptorRef& Interceptor, const FName& MessageType )
 {
-	FMessageTracerMessageInfoPtr MessageInfo = MessageInfos.FindRef(Context);
-
-	if (!MessageInfo.IsValid())
+	if (!Running)
 	{
 		return;
 	}
 
-	FMessageTracerEndpointInfoPtr& EndpointInfo = RecipientsToEndpointInfos.FindOrAdd(RecipientId);
+	double Timestamp = FPlatformTime::Seconds();
 
-	if (!EndpointInfo.IsValid())
+	Traces.Enqueue(
+		[=]() {
+			// @todo gmp: trace removed message interceptors
+		}
+	);
+}
+
+
+void FMessageTracer::TraceRemovedRecipient( const FMessageAddress& Address )
+{
+	if (!Running)
 	{
 		return;
 	}
 
-	// update message information
-	FMessageTracerDispatchStateRef DispatchState = MakeShareable(new FMessageTracerDispatchState());
+	double Timestamp = FPlatformTime::Seconds();
+
+	Traces.Enqueue(
+		[=]() {
+			FMessageTracerEndpointInfoPtr EndpointInfo = AddressesToEndpointInfos.FindRef(Address);
 
-	DispatchState->DispatchLatency = TimeSeconds - MessageInfo->TimeSent;
-	DispatchState->DispatchType = Async ? EMessageTracerDispatchTypes::TaskGraph : EMessageTracerDispatchTypes::Direct;
-	DispatchState->EndpointInfo = EndpointInfo;
-	DispatchState->TimeDispatched = TimeSeconds;
-	DispatchState->TimeHandled = 0.0;
+			if (!EndpointInfo.IsValid())
+			{
+				return;
+			}
 
-	MessageInfo->DispatchStates.Add(EndpointInfo, DispatchState);
+			// update endpoint information
+			FMessageTracerAddressInfoPtr AddressInfo = EndpointInfo->AddressInfos.FindRef(Address);
 
-	// update database
-	EndpointInfo->ReceivedMessages.Add(MessageInfo);
+			if (AddressInfo.IsValid())
+			{
+				AddressInfo->TimeUnregistered = Timestamp;
+			}
+		}
+	);
 }
 
 
-void FMessageTracer::ProcessHandledMessage( IMessageContextRef Context, double TimeSeconds, FGuid RecipientId )
+void FMessageTracer::TraceRemovedSubscription( const IMessageSubscriptionRef& Subscription, const FName& MessageType )
 {
-	FMessageTracerMessageInfoPtr MessageInfo = MessageInfos.FindRef(Context);
-
-	if (!MessageInfo.IsValid())
+	if (!Running)
 	{
 		return;
 	}
 
-	FMessageTracerEndpointInfoPtr EndpointInfo = RecipientsToEndpointInfos.FindRef(RecipientId);
+	double Timestamp = FPlatformTime::Seconds();
+
+	Traces.Enqueue(
+		[=]() {
+			// @todo gmp: trace removed message subscriptions
+		}
+	);
+}
+
 
-	if (!EndpointInfo.IsValid())
+void FMessageTracer::TraceRoutedMessage( const IMessageContextRef& Context )
+{
+	if (!Running)
 	{
 		return;
 	}
 
-	// update message information
-	FMessageTracerDispatchStatePtr DispatchState = MessageInfo->DispatchStates.FindRef(EndpointInfo);
+	double Timestamp = FPlatformTime::Seconds();
 
-	if (DispatchState.IsValid())
+	if (ShouldBreak(Context))
 	{
-		DispatchState->TimeHandled = TimeSeconds;
+		Breaking = true;
+		ContinueEvent->Wait();
 	}
-}
 
+	Traces.Enqueue(
+		[=]() {
+			// update message information
+			FMessageTracerMessageInfoPtr MessageInfo = MessageInfos.FindRef(Context);
 
-void FMessageTracer::ProcessRemovedInterceptor( IMessageInterceptorRef Interceptor, FName MessageType, double TimeSeconds )
-{
-
+			if (MessageInfo.IsValid())
+			{
+				MessageInfo->TimeRouted = Timestamp;
+			}
+		}
+	);
 }
 
 
-void FMessageTracer::ProcessRemovedRecipient( FMessageAddress Address, double TimeSeconds )
+void FMessageTracer::TraceSentMessage( const IMessageContextRef& Context )
 {
-	FMessageTracerEndpointInfoPtr EndpointInfo = AddressesToEndpointInfos.FindRef(Address);
-
-	if (!EndpointInfo.IsValid())
+	if (!Running)
 	{
 		return;
 	}
 
-	// update endpoint information
-	FMessageTracerAddressInfoPtr AddressInfo = EndpointInfo->AddressInfos.FindRef(Address);
+	double Timestamp = FPlatformTime::Seconds();
 
-	if (AddressInfo.IsValid())
-	{
-		AddressInfo->TimeUnregistered = TimeSeconds;
-	}
+	Traces.Enqueue(
+		[=]() {
+			// look up endpoint info
+			FMessageTracerEndpointInfoPtr EndpointInfo = AddressesToEndpointInfos.FindRef(Context->GetSender());
+
+			if (!EndpointInfo.IsValid())
+			{
+				return;
+			}
+
+			// create message info
+			FMessageTracerMessageInfoRef MessageInfo = MakeShareable(new FMessageTracerMessageInfo());
+	
+			MessageInfo->Context = Context;
+			MessageInfo->SenderInfo = EndpointInfo;
+			MessageInfo->TimeRouted = 0.0;
+			MessageInfo->TimeSent = Timestamp;
+			MessageInfos.Add(Context, MessageInfo);
+
+			// add message type
+			FMessageTracerTypeInfoPtr& TypeInfo = MessageTypes.FindOrAdd(Context->GetMessageType());
+
+			if (!TypeInfo.IsValid())
+			{
+				TypeInfo = MakeShareable(new FMessageTracerTypeInfo());
+				TypeInfo->TypeName = Context->GetMessageType();
+
+				TypeAddedDelegate.Broadcast(TypeInfo.ToSharedRef());
+			}
+
+			TypeInfo->Messages.Add(MessageInfo);
+
+			// update database
+			EndpointInfo->SentMessages.Add(MessageInfo);
+			MessageInfo->TypeInfo = TypeInfo;
+
+			MessagesAddedDelegate.Broadcast(MessageInfo);
+		}
+	);
 }
 
 
-void FMessageTracer::ProcessRemovedSubscription( FName MessageType, double TimeSeconds )
+/* IMessageTracer interface
+ *****************************************************************************/
+
+int32 FMessageTracer::GetEndpoints( TArray<FMessageTracerEndpointInfoPtr>& OutEndpoints ) const
 {
+	RecipientsToEndpointInfos.GenerateValueArray(OutEndpoints);
 
+	return OutEndpoints.Num();
 }
 
 
-void FMessageTracer::ProcessRoutedMessage( IMessageContextRef Context, double TimeSeconds )
+int32 FMessageTracer::GetMessages( TArray<FMessageTracerMessageInfoPtr>& OutMessages ) const
 {
-	// update message information
-	FMessageTracerMessageInfoPtr MessageInfo = MessageInfos.FindRef(Context);
+	MessageInfos.GenerateValueArray(OutMessages);
 
-	if (MessageInfo.IsValid())
-	{
-		MessageInfo->TimeRouted = TimeSeconds;
-	}
+	return OutMessages.Num();
 }
 
 
-void FMessageTracer::ProcessSentMessage( IMessageContextRef Context, double TimeSeconds )
+int32 FMessageTracer::GetMessageTypes( TArray<FMessageTracerTypeInfoPtr>& OutTypes ) const
 {
-	FMessageTracerEndpointInfoPtr EndpointInfo = AddressesToEndpointInfos.FindRef(Context->GetSender());
-
-	if (!EndpointInfo.IsValid())
-	{
-		return;
-	}
+	MessageTypes.GenerateValueArray(OutTypes);
 
-	// create message info
-	FMessageTracerMessageInfoRef MessageInfo = MakeShareable(new FMessageTracerMessageInfo());
-	
-	MessageInfo->Context = Context;
-	MessageInfo->SenderInfo = EndpointInfo;
-	MessageInfo->TimeRouted = 0.0;
-	MessageInfo->TimeSent = TimeSeconds;
-	MessageInfos.Add(Context, MessageInfo);
+	return OutTypes.Num();
+}
 
-	// add message type
-	FMessageTracerTypeInfoPtr& TypeInfo = MessageTypes.FindOrAdd(Context->GetMessageType());
 
-	if (!TypeInfo.IsValid())
+bool FMessageTracer::Tick( float DeltaTime )
+{
+	if (ResetPending)
 	{
-		TypeInfo = MakeShareable(new FMessageTracerTypeInfo());
-		TypeInfo->TypeName = Context->GetMessageType();
+		ResetMessages();
+		ResetPending = false;
+	}
 
-		TypeAddedDelegate.Broadcast(TypeInfo.ToSharedRef());
+	if (Traces.IsEmpty())
+	{
+		return false;
 	}
 
-	TypeInfo->Messages.Add(MessageInfo);
+	// process new traces
+	TFunction<void()> Trace;
 
-	// update database
-	EndpointInfo->SentMessages.Add(MessageInfo);
-	MessageInfo->TypeInfo = TypeInfo;
+	while (Traces.Dequeue(Trace))
+	{
+		Trace();
+	}
 
-	MessagesAddedDelegate.Broadcast(MessageInfo);
+	return true;
 }
 
 
+/* FMessageTracer implementation
+ *****************************************************************************/
+
 void FMessageTracer::ResetMessages()
 {
 	MessageInfos.Reset();
diff --git a/Engine/Source/Runtime/Messaging/Private/Bus/MessageTracer.h b/Engine/Source/Runtime/Messaging/Private/Bus/MessageTracer.h
--- a/Engine/Source/Runtime/Messaging/Private/Bus/MessageTracer.h
+++ b/Engine/Source/Runtime/Messaging/Private/Bus/MessageTracer.h
@@ -16,35 +16,13 @@
 class FMessageTracer
 	: public IMessageTracer
 {
-	/** Holds debug information for message recipients. */
-	struct FRecipientInfo
-	{
-		/** Holds the recipient's unique identifier. */
-		FGuid Id;
-
-		/** Holds the recipient's name. */
-		FName Name;
-
-		/** Holds a flag indicating whether the recipient is a remote endpoint. */
-		bool Remote;
-
-		/** Creates and initializes a new instance. */
-		FRecipientInfo( const FMessageAddress& Address, const IReceiveMessagesRef& Recipient )
-			: Id(Recipient->GetRecipientId())
-			, Name(Recipient->GetDebugName())
-			, Remote(Recipient->IsRemote())
-		{ }
-	};
-
 public:
 
-	DECLARE_DELEGATE(TraceDelegate)
-
 	/** Default constructor. */
 	FMessageTracer();
 
-	/** Destructor. */
-	~FMessageTracer();
+	/** Virtual destructor. */
+	virtual ~FMessageTracer();
 
 public:
 
@@ -54,11 +32,7 @@ public:
 	 * @param Interceptor The added interceptor.
 	 * @param MessageType The type of messages being intercepted.
 	 */
-	void TraceAddedInterceptor( const IMessageInterceptorRef& Interceptor, const FName& MessageType )
-	{
-		FString Name = TEXT("@todo");
-		Traces.Enqueue(FSimpleDelegate::CreateRaw(this, &FMessageTracer::ProcessAddedInterceptor, Name, MessageType, FPlatformTime::Seconds()));
-	}
+	void TraceAddedInterceptor( const IMessageInterceptorRef& Interceptor, const FName& MessageType );
 
 	/**
 	 * Notifies the tracer that a message recipient has been added to the message bus.
@@ -66,25 +40,14 @@ public:
 	 * @param Address The address of the added recipient.
 	 * @param Recipient The added recipient.
 	 */
-	void TraceAddedRecipient( const FMessageAddress& Address, const IReceiveMessagesRef& Recipient )
-	{
-		Traces.Enqueue(FSimpleDelegate::CreateRaw(this, &FMessageTracer::ProcessAddedRecipient, Address, FRecipientInfo(Address, Recipient), FPlatformTime::Seconds()));
-	}
+	void TraceAddedRecipient( const FMessageAddress& Address, const IReceiveMessagesRef& Recipient );
 
 	/**
 	 * Notifies the tracer that a message subscription has been added to the message bus.
 	 *
 	 * @param Subscription The added subscription.
 	 */
-	void TraceAddedSubscription( const IMessageSubscriptionRef& Subscription )
-	{
-		if (!Running)
-		{
-			return;
-		}
-
-		// @todo gmp: trace added subscriptions
-	}
+	void TraceAddedSubscription( const IMessageSubscriptionRef& Subscription );
 
 	/**
 	 * Notifies the tracer that a message has been dispatched.
@@ -93,15 +56,7 @@ public:
 	 * @param Recipient The message recipient.
 	 * @param Async Whether the message was dispatched asynchronously.
 	 */
-	void TraceDispatchedMessage( const IMessageContextRef& Context, const IReceiveMessagesRef& Recipient, bool Async )
-	{
-		if (!Running)
-		{
-			return;
-		}
-
-		Traces.Enqueue(FSimpleDelegate::CreateRaw(this, &FMessageTracer::ProcessDispatchedMessage, Context, FPlatformTime::Seconds(), Recipient->GetRecipientId(), Async));
-	}
+	void TraceDispatchedMessage( const IMessageContextRef& Context, const IReceiveMessagesRef& Recipient, bool Async );
 
 	/**
 	 * Notifies the tracer that a message has been handled.
@@ -109,15 +64,7 @@ public:
 	 * @param Context The context of the dispatched message.
 	 * @param Recipient The message recipient that handled the message.
 	 */
-	void TraceHandledMessage( const IMessageContextRef& Context, const IReceiveMessagesRef& Recipient )
-	{
-		if (!Running)
-		{
-			return;
-		}
-
-		Traces.Enqueue(FSimpleDelegate::CreateRaw(this, &FMessageTracer::ProcessHandledMessage, Context, FPlatformTime::Seconds(), Recipient->GetRecipientId()));
-	}
+	void TraceHandledMessage( const IMessageContextRef& Context, const IReceiveMessagesRef& Recipient );
 
 	/**
 	 * Notifies the tracer that a message has been intercepted.
@@ -125,14 +72,7 @@ public:
 	 * @param Context The context of the intercepted message.
 	 * @param Interceptor The interceptor.
 	 */
-	void TraceInterceptedMessage( const IMessageContextRef& Context, const IMessageInterceptorRef& Interceptor )
-	{
-		if (Running)
-		{
-		}
-
-		// @todo gmp: trace intercepted messages
-	}
+	void TraceInterceptedMessage( const IMessageContextRef& Context, const IMessageInterceptorRef& Interceptor );
 
 	/**
 	 * Notifies the tracer that a message interceptor has been removed from the message bus.
@@ -140,30 +80,14 @@ public:
 	 * @param Interceptor The removed interceptor.
 	 * @param MessageType The type of messages that is no longer being intercepted.
 	 */
-	void TraceRemovedInterceptor( const IMessageInterceptorRef& Interceptor, const FName& MessageType )
-	{
-		if (!Running)
-		{
-			return;
-		}
-
-		// @todo gmp: trace removed interceptors
-	}
+	void TraceRemovedInterceptor( const IMessageInterceptorRef& Interceptor, const FName& MessageType );
 
 	/**
 	 * Notifies the tracer that a recipient has been removed from the message bus.
 	 *
 	 * @param Address The address of the removed recipient.
 	 */
-	void TraceRemovedRecipient( const FMessageAddress& Address )
-	{
-		if (!Running)
-		{
-			return;
-		}
-
-		// @todo gmp: trace removed recipients
-	}
+	void TraceRemovedRecipient( const FMessageAddress& Address );
 
 	/**
 	 * Notifies the tracer that a message subscription has been removed from the message bus.
@@ -171,51 +95,21 @@ public:
 	 * @param Subscriber The removed subscriber.
 	 * @param MessageType The type of messages no longer being subscribed to.
 	 */
-	void TraceRemovedSubscription( const IMessageSubscriptionRef& Subscription, const FName& MessageType )
-	{
-		if (!Running)
-		{
-			return;
-		}
-
-		// @todo gmp: trace removed subscriptions
-	}
+	void TraceRemovedSubscription( const IMessageSubscriptionRef& Subscription, const FName& MessageType );
 
 	/**
 	 * Notifies the tracer that a message has been routed.
 	 *
 	 * @param Context The context of the routed message.
 	 */
-	void TraceRoutedMessage( const IMessageContextRef& Context )
-	{
-		if (!Running)
-		{
-			return;
-		}
-
-		if (ShouldBreak(Context))
-		{
-			Breaking = true;
-			ContinueEvent->Wait();
-		}
-
-		Traces.Enqueue(FSimpleDelegate::CreateRaw(this, &FMessageTracer::ProcessRoutedMessage, Context, FPlatformTime::Seconds()));
-	}
+	void TraceRoutedMessage( const IMessageContextRef& Context );
 
 	/**
 	 * Notifies the tracer that a message has been sent.
 	 *
 	 * @param Context The context of the sent message.
 	 */
-	void TraceSentMessage( const IMessageContextRef& Context )
-	{
-		if (!Running)
-		{
-			return;
-		}
-
-		Traces.Enqueue(FSimpleDelegate::CreateRaw(this, &FMessageTracer::ProcessSentMessage, Context, FPlatformTime::Seconds()));
-	}
+	void TraceSentMessage( const IMessageContextRef& Context );
 
 public:
 
@@ -319,102 +213,6 @@ public:
 
 protected:
 
-	/**
-	 * Enqueues a trace action for synchronized processing.
-	 *
-	 * @param Trace The action to enqueue.
-	 */
-	FORCEINLINE void EnqueueTrace( TraceDelegate Trace )
-	{
-		Traces.Enqueue(Trace);
-	}
-
-	/**
-	 * Processes traces for added message interceptors.
-	 *
-	 * @param Name The name of the interceptor.
-	 * @param MessageType The type of messages being intercepted.
-	 * @param TimeSecond The time at which the interceptor was added.
-	 */
-	void ProcessAddedInterceptor( FString Name, FName MessageType, double TimeSeconds );
-
-	/**
-	 * Processes traces for added message recipients.
-	 *
-	 * @param Address The address of the added recipient.
-	 * @param RecipientInfo Information about the recipient (name, ID, etc.)
-	 * @param TimeSecond The time at which the recipient was added.
-	 */
-	void ProcessAddedRecipient( FMessageAddress Address, FRecipientInfo RecipientInfo, double TimeSeconds );
-
-	/**
-	 * Processes traces for added message subscriptions.
-	 *
-	 * @param TimeSecond The time at which the subscription was added.
-	 */
-	void ProcessAddedSubscriptionTrace( double TimeSeconds );
-
-	/**
-	 * Processes traces for dispatched messages.
-	 *
-	 * @param Context The context of the dispatched message.
-	 * @param TimeSecond The time at which the message was dispatched.
-	 * @param RecipientId The recipient's unique identifier.
-	 * @param Async Whether the message was dispatched asynchronously.
-	 */
-	void ProcessDispatchedMessage( IMessageContextRef Context, double TimeSeconds, FGuid RecipientId, bool Async );
-
-	/**
-	 * Processes traces for handled messages.
-	 *
-	 * @param Context The context of the handled message.
-	 * @param TimeSecond The time at which the message was handled.
-	 * @param RecipientId The recipient's unique identifier.
-	 */
-	void ProcessHandledMessage( IMessageContextRef Context, double TimeSeconds, FGuid RecipientId );
-
-	/**
-	 * Processes traces for removed message interceptors.
-	 *
-	 * @param Interceptor The removed interceptor.
-	 * @param MessageType The type of messages no longer being intercepted.
-	 * @param TimeSecond The time at which the interceptor was removed.
-	 */
-	void ProcessRemovedInterceptor( IMessageInterceptorRef Interceptor, FName MessageType, double TimeSeconds );
-
-	/**
-	 * Processes traces for removed message recipients.
-	 *
-	 * @param Address The address of the removed recipient.
-	 * @param TimeSecond The time at which the recipient was removed.
-	 */
-	void ProcessRemovedRecipient( FMessageAddress Address, double TimeSeconds );
-
-	/**
-	 * Processes traces for removed message subscriptions.
-	 *
-	 * @param MessageType The type of messages no longer being subscribed to.
-	 * @param TimeSecond The time at which the subscription was removed.
-	 */
-	void ProcessRemovedSubscription( FName MessageType, double TimeSeconds );
-
-	/**
-	 * Processes traces for routed messages.
-	 *
-	 * @param Context The context of the routed message.
-	 * @param TimeSecond The time at which the message was routed.
-	 */
-	void ProcessRoutedMessage( IMessageContextRef Context, double TimeSeconds );
-
-	/**
-	 * Processes traces for sent messages.
-	 *
-	 * @param Context The context of the sent message.
-	 * @param SenderThread The name of the thread from which the message was sent.
-	 * @param TimeSecond The time at which the message was sent.
-	 */
-	void ProcessSentMessage( IMessageContextRef Context, double TimeSeconds );
-
 	/**
 	 * Resets traced messages.
 	 */
@@ -460,7 +258,7 @@ private:
 	bool Running;
 
 	/** Holds the trace actions queue. */
-	TQueue<TraceDelegate, EQueueMode::Mpsc> Traces;
+	TQueue<TFunction<void()>, EQueueMode::Mpsc> Traces;
 
 private:
 
