Fix Mac compile error for portal project

diff --git a/Engine/Source/Runtime/Messaging/Public/IMessageHandler.h b/Engine/Source/Runtime/Messaging/Public/IMessageHandler.h
--- a/Engine/Source/Runtime/Messaging/Public/IMessageHandler.h
+++ b/Engine/Source/Runtime/Messaging/Public/IMessageHandler.h
@@ -3,7 +3,7 @@
 #pragma once
 
 #include "CoreUObject.h"
-
+#include "IMessageContext.h"
 
 class IMessageContext;
 
