Messaging: Refactored IMessageHandler.h so that it doesn't have header dependencies

diff --git a/Engine/Source/Runtime/Messaging/Public/Helpers/MessageEndpointBuilder.h b/Engine/Source/Runtime/Messaging/Public/Helpers/MessageEndpointBuilder.h
--- a/Engine/Source/Runtime/Messaging/Public/Helpers/MessageEndpointBuilder.h
+++ b/Engine/Source/Runtime/Messaging/Public/Helpers/MessageEndpointBuilder.h
@@ -6,6 +6,7 @@
 #include "IMessageHandler.h"
 #include "IMessagingModule.h"
 #include "MessageEndpoint.h"
+#include "MessageHandlers.h"
 #include "TaskGraphInterfaces.h"
 
 
@@ -72,7 +73,7 @@ public:
 	}
 
 	/**
-	 * Adds a message handler for the given type of messages.
+	 * Adds a message handler for the given type of messages (via TFunction object).
 	 *
 	 * It is legal to configure multiple handlers for the same message type. Each
 	 * handler will be executed when a message of the specified type is received.
diff --git a/Engine/Source/Runtime/Messaging/Public/Helpers/MessageHandlers.h b/Engine/Source/Runtime/Messaging/Public/Helpers/MessageHandlers.h
new file mode 100644
--- /dev/null
+++ b/Engine/Source/Runtime/Messaging/Public/Helpers/MessageHandlers.h
@@ -0,0 +1,109 @@
+// Copyright 1998-2015 Epic Games, Inc. All Rights Reserved.
+
+#pragma once
+
+#include "CoreUObject.h"
+#include "IMessageContext.h"
+
+
+/**
+ * Template for handlers of one specific message type (via raw function pointers).
+ *
+ * @param MessageType The type of message to handle.
+ * @param HandlerType The type of the handler class.
+ */
+template<typename MessageType, typename HandlerType>
+class TRawMessageHandler
+	: public IMessageHandler
+{
+public:
+
+	/** Type definition for function pointers that are compatible with this TRawMessageHandler. */
+	typedef void (HandlerType::*FuncType)( const MessageType&, const TSharedRef<IMessageContext, ESPMode::ThreadSafe>& );
+
+public:
+
+	/**
+	 * Creates and initializes a new message handler.
+	 *
+	 * @param InHandler The object that will handle the messages.
+	 * @param InHandlerFunc The object's message handling function.
+	 */
+	TRawMessageHandler( HandlerType* InHandler, FuncType InHandlerFunc )
+		: Handler(InHandler)
+		, HandlerFunc(InHandlerFunc)
+	{
+		check(InHandler != nullptr);
+	}
+
+	/** Virtual destructor. */
+	~TRawMessageHandler() { }
+
+public:
+
+	// IMessageHandler interface
+	
+	virtual void HandleMessage( const TSharedRef<IMessageContext, ESPMode::ThreadSafe>& Context ) override
+	{
+		if (Context->GetMessageType() == MessageType::StaticStruct()->GetFName())
+		{
+			(Handler->*HandlerFunc)(*static_cast<const MessageType*>(Context->GetMessage()), Context);
+		}	
+	}
+	
+private:
+
+	/** Holds a pointer to the object handling the messages. */
+	HandlerType* Handler;
+
+	/** Holds a pointer to the actual handler function. */
+	FuncType HandlerFunc;
+};
+
+
+/**
+ * Template for handlers of one specific message type (via function objects).
+ *
+ * @param MessageType The type of message to handle.
+ */
+template<typename MessageType>
+class TFunctionMessageHandler
+	: public IMessageHandler
+{
+public:
+
+	/** Type definition for function objects that are compatible with this TFunctionMessageHandler. */
+	typedef TFunction<void(const MessageType&, const TSharedRef<IMessageContext, ESPMode::ThreadSafe>&)> FuncType;
+
+public:
+
+	/**
+	 * Creates and initializes a new message handler.
+	 *
+	 * @param InHandler The object that will handle the messages.
+	 * @param InHandlerFunc The object's message handling function.
+	 */
+	TFunctionMessageHandler( FuncType InHandlerFunc )
+		: HandlerFunc(InHandlerFunc)
+	{ }
+
+	/** Virtual destructor. */
+	~TFunctionMessageHandler() { }
+
+public:
+
+	// IMessageHandler interface
+	
+	virtual void HandleMessage( const TSharedRef<IMessageContext, ESPMode::ThreadSafe>& Context ) override
+	{
+		if (Context->GetMessageType() == MessageType::StaticStruct()->GetFName())
+		{
+			HandlerFunc(*static_cast<const MessageType*>(Context->GetMessage()), Context);
+		}	
+	}
+	
+private:
+
+	/** Holds a pointer to the actual handler function. */
+	FuncType HandlerFunc;
+};
diff --git a/Engine/Source/Runtime/Messaging/Public/IMessageHandler.h b/Engine/Source/Runtime/Messaging/Public/IMessageHandler.h
--- a/Engine/Source/Runtime/Messaging/Public/IMessageHandler.h
+++ b/Engine/Source/Runtime/Messaging/Public/IMessageHandler.h
@@ -2,11 +2,6 @@
 
 #pragma once
 
-#include "CoreUObject.h"
-#include "IMessageContext.h"
-
-class IMessageContext;
-
 
 /**
  * Interface for message handlers.
@@ -34,106 +29,3 @@ typedef TSharedPtr<IMessageHandler, ESPMode::ThreadSafe> IMessageHandlerPtr;
 
 /** Type definition for shared references to instances of IMessageHandler. */
 typedef TSharedRef<IMessageHandler, ESPMode::ThreadSafe> IMessageHandlerRef;
-
-
-/**
- * Template for handlers of one specific message type (via raw function pointers).
- *
- * @param MessageType The type of message to handle.
- * @param HandlerType The type of the handler class.
- */
-template<typename MessageType, typename HandlerType>
-class TRawMessageHandler
-	: public IMessageHandler
-{
-public:
-
-	/** Type definition for function pointers that are compatible with this TRawMessageHandler. */
-	typedef void (HandlerType::*FuncType)( const MessageType&, const TSharedRef<IMessageContext, ESPMode::ThreadSafe>& );
-
-public:
-
-	/**
-	 * Creates and initializes a new message handler.
-	 *
-	 * @param InHandler The object that will handle the messages.
-	 * @param InHandlerFunc The object's message handling function.
-	 */
-	TRawMessageHandler( HandlerType* InHandler, FuncType InHandlerFunc )
-		: Handler(InHandler)
-		, HandlerFunc(InHandlerFunc)
-	{
-		check(InHandler != nullptr);
-	}
-
-	/** Virtual destructor. */
-	~TRawMessageHandler() { }
-
-public:
-
-	// IMessageHandler interface
-	
-	virtual void HandleMessage( const TSharedRef<IMessageContext, ESPMode::ThreadSafe>& Context ) override
-	{
-		if (Context->GetMessageType() == MessageType::StaticStruct()->GetFName())
-		{
-			(Handler->*HandlerFunc)(*static_cast<const MessageType*>(Context->GetMessage()), Context);
-		}	
-	}
-	
-private:
-
-	/** Holds a pointer to the object handling the messages. */
-	HandlerType* Handler;
-
-	/** Holds a pointer to the actual handler function. */
-	FuncType HandlerFunc;
-};
-
-
-/**
- * Template for handlers of one specific message type (via function objects).
- *
- * @param MessageType The type of message to handle.
- */
-template<typename MessageType>
-class TFunctionMessageHandler
-	: public IMessageHandler
-{
-public:
-
-	/** Type definition for function objects that are compatible with this TFunctionMessageHandler. */
-	typedef TFunction<void(const MessageType&, const TSharedRef<IMessageContext, ESPMode::ThreadSafe>&)> FuncType;
-
-public:
-
-	/**
-	 * Creates and initializes a new message handler.
-	 *
-	 * @param InHandler The object that will handle the messages.
-	 * @param InHandlerFunc The object's message handling function.
-	 */
-	TFunctionMessageHandler( FuncType InHandlerFunc )
-		: HandlerFunc(InHandlerFunc)
-	{ }
-
-	/** Virtual destructor. */
-	~TFunctionMessageHandler() { }
-
-public:
-
-	// IMessageHandler interface
-	
-	virtual void HandleMessage( const TSharedRef<IMessageContext, ESPMode::ThreadSafe>& Context ) override
-	{
-		if (Context->GetMessageType() == MessageType::StaticStruct()->GetFName())
-		{
-			HandlerFunc(*static_cast<const MessageType*>(Context->GetMessage()), Context);
-		}	
-	}
-	
-private:
-
-	/** Holds a pointer to the actual handler function. */
-	FuncType HandlerFunc;
-};
