Messaging: Added a doc comment about message bus taking ownership of messages

diff --git a/Engine/Source/Runtime/Messaging/Public/Helpers/MessageEndpoint.h b/Engine/Source/Runtime/Messaging/Public/Helpers/MessageEndpoint.h
--- a/Engine/Source/Runtime/Messaging/Public/Helpers/MessageEndpoint.h
+++ b/Engine/Source/Runtime/Messaging/Public/Helpers/MessageEndpoint.h
@@ -41,6 +41,9 @@ DECLARE_DELEGATE_RetVal_OneParam(bool, FOnMessageEndpointReceiveMessage, const I
  * Endpoints that are destroyed or receive messages on non-Game threads should use the static function
  * FMessageEndpoint::SafeRelease() to dispose of the endpoint. This will ensure that there are no race
  * conditions between endpoint destruction and the receiving of messages.
+ *
+ * The underlying message bus will take ownership of all sent and published message objects. The memory
+ * held by the messages must therefore NOT be freed by the caller.
  */
 class FMessageEndpoint
 	: public TSharedFromThis<FMessageEndpoint, ESPMode::ThreadSafe>
diff --git a/Engine/Source/Runtime/Messaging/Public/IMessageBus.h b/Engine/Source/Runtime/Messaging/Public/IMessageBus.h
--- a/Engine/Source/Runtime/Messaging/Public/IMessageBus.h
+++ b/Engine/Source/Runtime/Messaging/Public/IMessageBus.h
@@ -134,7 +134,8 @@ public:
 	/**
 	 * Sends a message to subscribed recipients.
 	 *
-	 * The message is published with a default time to live of Process.
+	 * The bus takes over ownership of the message's memory.
+	 * It must NOT be freed by the caller.
 	 *
 	 * @param Message The message to publish.
 	 * @param TypeInfo The message's type information.
@@ -158,6 +159,9 @@ public:
 	/**
 	 * Sends a message to multiple recipients.
 	 *
+	 * The bus takes over ownership of the message's memory.
+	 * It must NOT be freed by the caller.
+	 *
 	 * @param Message The message to send.
 	 * @param TypeInfo The message's type information.
 	 * @param Attachment The binary data to attach to the message.
