Messaging: Fixed message tracer not processing traces unless MessagingDebugger is open (UE-13752)

diff --git a/Engine/Plugins/Messaging/MessagingDebugger/Source/MessagingDebugger/Private/Widgets/SMessagingDebugger.cpp b/Engine/Plugins/Messaging/MessagingDebugger/Source/MessagingDebugger/Private/Widgets/SMessagingDebugger.cpp
--- a/Engine/Plugins/Messaging/MessagingDebugger/Source/MessagingDebugger/Private/Widgets/SMessagingDebugger.cpp
+++ b/Engine/Plugins/Messaging/MessagingDebugger/Source/MessagingDebugger/Private/Widgets/SMessagingDebugger.cpp
@@ -219,15 +219,6 @@ void SMessagingDebugger::Construct( const FArguments& InArgs, const TSharedRef<S
 }
 
 
-/* SCompoundWidget overrides
- *****************************************************************************/
-
-void SMessagingDebugger::Tick( const FGeometry& AllottedGeometry, const double InCurrentTime, const float InDeltaTime )
-{
-	MessageTracer->Tick( InDeltaTime );
-}
-
-
 /* SMessagingDebugger implementation
  *****************************************************************************/
 
diff --git a/Engine/Plugins/Messaging/MessagingDebugger/Source/MessagingDebugger/Private/Widgets/SMessagingDebugger.h b/Engine/Plugins/Messaging/MessagingDebugger/Source/MessagingDebugger/Private/Widgets/SMessagingDebugger.h
--- a/Engine/Plugins/Messaging/MessagingDebugger/Source/MessagingDebugger/Private/Widgets/SMessagingDebugger.h
+++ b/Engine/Plugins/Messaging/MessagingDebugger/Source/MessagingDebugger/Private/Widgets/SMessagingDebugger.h
@@ -32,11 +32,6 @@ public:
 	 */
 	void Construct( const FArguments& InArgs, const TSharedRef<SDockTab>& ConstructUnderMajorTab, const TSharedPtr<SWindow>& ConstructUnderWindow, const IMessageTracerRef& InMessageTracer, const TSharedRef<ISlateStyle>& InStyle );
 
-public:
-
-	// SCompoundWidget overrides
-	virtual void Tick( const FGeometry& AllottedGeometry, const double InCurrentTime, const float InDeltaTime ) override;
-
 protected:
 
 	/**
diff --git a/Engine/Source/Runtime/Messaging/Private/Bus/MessageTracer.cpp b/Engine/Source/Runtime/Messaging/Private/Bus/MessageTracer.cpp
--- a/Engine/Source/Runtime/Messaging/Private/Bus/MessageTracer.cpp
+++ b/Engine/Source/Runtime/Messaging/Private/Bus/MessageTracer.cpp
@@ -12,11 +12,13 @@ FMessageTracer::FMessageTracer()
 	, Running(false)
 {
 	ContinueEvent = FPlatformProcess::GetSynchEventFromPool();
+	TickDelegateHandle = FTicker::GetCoreTicker().AddTicker(FTickerDelegate::CreateRaw(this, &FMessageTracer::Tick), 0.0f);
 }
 
 
 FMessageTracer::~FMessageTracer()
 {
+	FTicker::GetCoreTicker().RemoveTicker(TickDelegateHandle);
 	FPlatformProcess::ReturnSynchEventToPool(ContinueEvent);
 	ContinueEvent = nullptr;
 }
@@ -29,11 +31,9 @@ void FMessageTracer::TraceAddedInterceptor( const IMessageInterceptorRef& Interc
 {
 	double Timestamp = FPlatformTime::Seconds();
 
-	Traces.Enqueue(
-		[=]() {
-			// @todo gmp: trace added message interceptors
-		}
-	);
+	Traces.Enqueue([=]() {
+		// @todo gmp: trace added message interceptors
+	});
 }
 
 
@@ -41,30 +41,28 @@ void FMessageTracer::TraceAddedRecipient( const FMessageAddress& Address, const
 {
 	double Timestamp = FPlatformTime::Seconds();
 
-	Traces.Enqueue(
-		[=]() {
-			// create endpoint information
-			FMessageTracerEndpointInfoPtr& EndpointInfo = RecipientsToEndpointInfos.FindOrAdd(Recipient->GetRecipientId());
+	Traces.Enqueue([=]() {
+		// create endpoint information
+		FMessageTracerEndpointInfoPtr& EndpointInfo = RecipientsToEndpointInfos.FindOrAdd(Recipient->GetRecipientId());
 
-			if (!EndpointInfo.IsValid())
-			{
-				EndpointInfo = MakeShareable(new FMessageTracerEndpointInfo());
-			}
+		if (!EndpointInfo.IsValid())
+		{
+			EndpointInfo = MakeShareable(new FMessageTracerEndpointInfo());
+		}
 
-			// initialize endpoint information
-			FMessageTracerAddressInfoRef AddressInfo = MakeShareable(new FMessageTracerAddressInfo());
-			AddressInfo->Address = Address;
-			AddressInfo->TimeRegistered = Timestamp;
-			AddressInfo->TimeUnregistered = 0;
+		// initialize endpoint information
+		FMessageTracerAddressInfoRef AddressInfo = MakeShareable(new FMessageTracerAddressInfo());
+		AddressInfo->Address = Address;
+		AddressInfo->TimeRegistered = Timestamp;
+		AddressInfo->TimeUnregistered = 0;
 
-			EndpointInfo->AddressInfos.Add(Address, AddressInfo);
-			EndpointInfo->Name = Recipient->GetDebugName();
-			EndpointInfo->Remote = Recipient->IsRemote();
+		EndpointInfo->AddressInfos.Add(Address, AddressInfo);
+		EndpointInfo->Name = Recipient->GetDebugName();
+		EndpointInfo->Remote = Recipient->IsRemote();
 
-			// add to address table
-			AddressesToEndpointInfos.Add(Address, EndpointInfo);
-		}
-	);
+		// add to address table
+		AddressesToEndpointInfos.Add(Address, EndpointInfo);
+	});
 }
 
 
@@ -77,11 +75,9 @@ void FMessageTracer::TraceAddedSubscription( const IMessageSubscriptionRef& Subs
 
 	double Timestamp = FPlatformTime::Seconds();
 
-	Traces.Enqueue(
-		[=]() {
-			// @todo gmp: trace added subscriptions
-		}
-	);
+	Traces.Enqueue([=]() {
+		// @todo gmp: trace added subscriptions
+	});
 }
 
 
@@ -94,39 +90,37 @@ void FMessageTracer::TraceDispatchedMessage( const IMessageContextRef& Context,
 
 	double Timestamp = FPlatformTime::Seconds();
 
-	Traces.Enqueue(
-		[=]() {
-			// look up message & endpoint info
-			FMessageTracerMessageInfoPtr MessageInfo = MessageInfos.FindRef(Context);
+	Traces.Enqueue([=]() {
+		// look up message & endpoint info
+		FMessageTracerMessageInfoPtr MessageInfo = MessageInfos.FindRef(Context);
 
-			if (!MessageInfo.IsValid())
-			{
-				return;
-			}
+		if (!MessageInfo.IsValid())
+		{
+			return;
+		}
 
-			FMessageTracerEndpointInfoPtr& EndpointInfo = RecipientsToEndpointInfos.FindOrAdd(Recipient->GetRecipientId());
+		FMessageTracerEndpointInfoPtr& EndpointInfo = RecipientsToEndpointInfos.FindOrAdd(Recipient->GetRecipientId());
 
-			if (!EndpointInfo.IsValid())
-			{
-				return;
-			}
+		if (!EndpointInfo.IsValid())
+		{
+			return;
+		}
 
-			// update message information
-			FMessageTracerDispatchStateRef DispatchState = MakeShareable(new FMessageTracerDispatchState());
+		// update message information
+		FMessageTracerDispatchStateRef DispatchState = MakeShareable(new FMessageTracerDispatchState());
 
-			DispatchState->DispatchLatency = Timestamp - MessageInfo->TimeSent;
-			DispatchState->DispatchType = Async ? EMessageTracerDispatchTypes::TaskGraph : EMessageTracerDispatchTypes::Direct;
-			DispatchState->EndpointInfo = EndpointInfo;
-			DispatchState->RecipientThread = Recipient->GetRecipientThread();
-			DispatchState->TimeDispatched = Timestamp;
-			DispatchState->TimeHandled = 0.0;
+		DispatchState->DispatchLatency = Timestamp - MessageInfo->TimeSent;
+		DispatchState->DispatchType = Async ? EMessageTracerDispatchTypes::TaskGraph : EMessageTracerDispatchTypes::Direct;
+		DispatchState->EndpointInfo = EndpointInfo;
+		DispatchState->RecipientThread = Recipient->GetRecipientThread();
+		DispatchState->TimeDispatched = Timestamp;
+		DispatchState->TimeHandled = 0.0;
 
-			MessageInfo->DispatchStates.Add(EndpointInfo, DispatchState);
+		MessageInfo->DispatchStates.Add(EndpointInfo, DispatchState);
 
-			// update database
-			EndpointInfo->ReceivedMessages.Add(MessageInfo);
-		}
-	);
+		// update database
+		EndpointInfo->ReceivedMessages.Add(MessageInfo);
+	});
 }
 
 
@@ -139,32 +133,30 @@ void FMessageTracer::TraceHandledMessage( const IMessageContextRef& Context, con
 
 	double Timestamp = FPlatformTime::Seconds();
 
-	Traces.Enqueue(
-		[=]() {
-			// look up message & endpoint info
-			FMessageTracerMessageInfoPtr MessageInfo = MessageInfos.FindRef(Context);
+	Traces.Enqueue([=]() {
+		// look up message & endpoint info
+		FMessageTracerMessageInfoPtr MessageInfo = MessageInfos.FindRef(Context);
 
-			if (!MessageInfo.IsValid())
-			{
-				return;
-			}
+		if (!MessageInfo.IsValid())
+		{
+			return;
+		}
 
-			FMessageTracerEndpointInfoPtr EndpointInfo = RecipientsToEndpointInfos.FindRef(Recipient->GetRecipientId());
+		FMessageTracerEndpointInfoPtr EndpointInfo = RecipientsToEndpointInfos.FindRef(Recipient->GetRecipientId());
 
-			if (!EndpointInfo.IsValid())
-			{
-				return;
-			}
+		if (!EndpointInfo.IsValid())
+		{
+			return;
+		}
 
-			// update message information
-			FMessageTracerDispatchStatePtr DispatchState = MessageInfo->DispatchStates.FindRef(EndpointInfo);
+		// update message information
+		FMessageTracerDispatchStatePtr DispatchState = MessageInfo->DispatchStates.FindRef(EndpointInfo);
 
-			if (DispatchState.IsValid())
-			{
-				DispatchState->TimeHandled = Timestamp;
-			}
+		if (DispatchState.IsValid())
+		{
+			DispatchState->TimeHandled = Timestamp;
 		}
-	);
+	});
 }
 
 
@@ -177,11 +169,9 @@ void FMessageTracer::TraceInterceptedMessage( const IMessageContextRef& Context,
 
 	double Timestamp = FPlatformTime::Seconds();
 
-	Traces.Enqueue(
-		[=]() {
-			// @todo gmp: trace intercepted messages
-		}
-	);		
+	Traces.Enqueue([=]() {
+		// @todo gmp: trace intercepted messages
+	});		
 }
 
 
@@ -194,11 +184,9 @@ void FMessageTracer::TraceRemovedInterceptor( const IMessageInterceptorRef& Inte
 
 	double Timestamp = FPlatformTime::Seconds();
 
-	Traces.Enqueue(
-		[=]() {
-			// @todo gmp: trace removed message interceptors
-		}
-	);
+	Traces.Enqueue([=]() {
+		// @todo gmp: trace removed message interceptors
+	});
 }
 
 
@@ -211,24 +199,22 @@ void FMessageTracer::TraceRemovedRecipient( const FMessageAddress& Address )
 
 	double Timestamp = FPlatformTime::Seconds();
 
-	Traces.Enqueue(
-		[=]() {
-			FMessageTracerEndpointInfoPtr EndpointInfo = AddressesToEndpointInfos.FindRef(Address);
+	Traces.Enqueue([=]() {
+		FMessageTracerEndpointInfoPtr EndpointInfo = AddressesToEndpointInfos.FindRef(Address);
 
-			if (!EndpointInfo.IsValid())
-			{
-				return;
-			}
+		if (!EndpointInfo.IsValid())
+		{
+			return;
+		}
 
-			// update endpoint information
-			FMessageTracerAddressInfoPtr AddressInfo = EndpointInfo->AddressInfos.FindRef(Address);
+		// update endpoint information
+		FMessageTracerAddressInfoPtr AddressInfo = EndpointInfo->AddressInfos.FindRef(Address);
 
-			if (AddressInfo.IsValid())
-			{
-				AddressInfo->TimeUnregistered = Timestamp;
-			}
+		if (AddressInfo.IsValid())
+		{
+			AddressInfo->TimeUnregistered = Timestamp;
 		}
-	);
+	});
 }
 
 
@@ -241,11 +227,9 @@ void FMessageTracer::TraceRemovedSubscription( const IMessageSubscriptionRef& Su
 
 	double Timestamp = FPlatformTime::Seconds();
 
-	Traces.Enqueue(
-		[=]() {
-			// @todo gmp: trace removed message subscriptions
-		}
-	);
+	Traces.Enqueue([=]() {
+		// @todo gmp: trace removed message subscriptions
+	});
 }
 
 
@@ -264,17 +248,15 @@ void FMessageTracer::TraceRoutedMessage( const IMessageContextRef& Context )
 		ContinueEvent->Wait();
 	}
 
-	Traces.Enqueue(
-		[=]() {
-			// update message information
-			FMessageTracerMessageInfoPtr MessageInfo = MessageInfos.FindRef(Context);
+	Traces.Enqueue([=]() {
+		// update message information
+		FMessageTracerMessageInfoPtr MessageInfo = MessageInfos.FindRef(Context);
 
-			if (MessageInfo.IsValid())
-			{
-				MessageInfo->TimeRouted = Timestamp;
-			}
+		if (MessageInfo.IsValid())
+		{
+			MessageInfo->TimeRouted = Timestamp;
 		}
-	);
+	});
 }
 
 
@@ -287,45 +269,43 @@ void FMessageTracer::TraceSentMessage( const IMessageContextRef& Context )
 
 	double Timestamp = FPlatformTime::Seconds();
 
-	Traces.Enqueue(
-		[=]() {
-			// look up endpoint info
-			FMessageTracerEndpointInfoPtr EndpointInfo = AddressesToEndpointInfos.FindRef(Context->GetSender());
+	Traces.Enqueue([=]() {
+		// look up endpoint info
+		FMessageTracerEndpointInfoPtr EndpointInfo = AddressesToEndpointInfos.FindRef(Context->GetSender());
 
-			if (!EndpointInfo.IsValid())
-			{
-				return;
-			}
+		if (!EndpointInfo.IsValid())
+		{
+			return;
+		}
 
-			// create message info
-			FMessageTracerMessageInfoRef MessageInfo = MakeShareable(new FMessageTracerMessageInfo());
+		// create message info
+		FMessageTracerMessageInfoRef MessageInfo = MakeShareable(new FMessageTracerMessageInfo());
 	
-			MessageInfo->Context = Context;
-			MessageInfo->SenderInfo = EndpointInfo;
-			MessageInfo->TimeRouted = 0.0;
-			MessageInfo->TimeSent = Timestamp;
-			MessageInfos.Add(Context, MessageInfo);
+		MessageInfo->Context = Context;
+		MessageInfo->SenderInfo = EndpointInfo;
+		MessageInfo->TimeRouted = 0.0;
+		MessageInfo->TimeSent = Timestamp;
+		MessageInfos.Add(Context, MessageInfo);
 
-			// add message type
-			FMessageTracerTypeInfoPtr& TypeInfo = MessageTypes.FindOrAdd(Context->GetMessageType());
+		// add message type
+		FMessageTracerTypeInfoPtr& TypeInfo = MessageTypes.FindOrAdd(Context->GetMessageType());
 
-			if (!TypeInfo.IsValid())
-			{
-				TypeInfo = MakeShareable(new FMessageTracerTypeInfo());
-				TypeInfo->TypeName = Context->GetMessageType();
+		if (!TypeInfo.IsValid())
+		{
+			TypeInfo = MakeShareable(new FMessageTracerTypeInfo());
+			TypeInfo->TypeName = Context->GetMessageType();
 
-				TypeAddedDelegate.Broadcast(TypeInfo.ToSharedRef());
-			}
+			TypeAddedDelegate.Broadcast(TypeInfo.ToSharedRef());
+		}
 
-			TypeInfo->Messages.Add(MessageInfo);
+		TypeInfo->Messages.Add(MessageInfo);
 
-			// update database
-			EndpointInfo->SentMessages.Add(MessageInfo);
-			MessageInfo->TypeInfo = TypeInfo;
+		// update database
+		EndpointInfo->SentMessages.Add(MessageInfo);
+		MessageInfo->TypeInfo = TypeInfo;
 
-			MessagesAddedDelegate.Broadcast(MessageInfo);
-		}
-	);
+		MessagesAddedDelegate.Broadcast(MessageInfo);
+	});
 }
 
 
@@ -364,17 +344,15 @@ bool FMessageTracer::Tick( float DeltaTime )
 		ResetPending = false;
 	}
 
-	if (Traces.IsEmpty())
-	{
-		return false;
-	}
-
 	// process new traces
-	TFunction<void()> Trace;
-
-	while (Traces.Dequeue(Trace))
+	if (!Traces.IsEmpty())
 	{
-		Trace();
+		TFunction<void()> Trace;
+
+		while (Traces.Dequeue(Trace))
+		{
+			Trace();
+		}
 	}
 
 	return true;
diff --git a/Engine/Source/Runtime/Messaging/Private/Bus/MessageTracer.h b/Engine/Source/Runtime/Messaging/Private/Bus/MessageTracer.h
--- a/Engine/Source/Runtime/Messaging/Private/Bus/MessageTracer.h
+++ b/Engine/Source/Runtime/Messaging/Private/Bus/MessageTracer.h
@@ -257,6 +257,9 @@ private:
 	/** Holds a flag indicating whether the tracer is running. */
 	bool Running;
 
+	/** Handle to the registered TickDelegate. */
+	FDelegateHandle TickDelegateHandle;
+
 	/** Holds the trace actions queue. */
 	TQueue<TFunction<void()>, EQueueMode::Mpsc> Traces;
 
