Messaging: Added catch-all handlers to message endpoints

diff --git a/Engine/Source/Runtime/Messaging/Public/Helpers/MessageEndpointBuilder.h b/Engine/Source/Runtime/Messaging/Public/Helpers/MessageEndpointBuilder.h
--- a/Engine/Source/Runtime/Messaging/Public/Helpers/MessageEndpointBuilder.h
+++ b/Engine/Source/Runtime/Messaging/Public/Helpers/MessageEndpointBuilder.h
@@ -62,7 +62,7 @@ public:
 	 * @param Handler The class handling the messages.
 	 * @param HandlerFunc The class function handling the messages.
 	 * @return This instance (for method chaining).
-	 * @see WithHandler
+	 * @see WithCatchall, WithHandler
 	 */
 	template<typename MessageType, typename HandlerType>
 	FMessageEndpointBuilder& Handling( HandlerType* Handler, typename TRawMessageHandler<MessageType, HandlerType>::FuncType HandlerFunc )
@@ -84,7 +84,7 @@ public:
 	 * @param MessageType The type of messages to handle.
 	 * @param Function The function object handling the messages.
 	 * @return This instance (for method chaining).
-	 * @see WithHandler
+	 * @see WithCatchall, WithHandler
 	 */
 	template<typename MessageType>
 	FMessageEndpointBuilder& Handling( typename TFunctionMessageHandler<MessageType>::FuncType HandlerFunc )
@@ -147,6 +147,50 @@ public:
 		return *this;
 	}
 
+	/**
+	 * Adds a message handler for the given type of messages (via raw function pointers).
+	 *
+	 * It is legal to configure multiple handlers for the same message type. Each
+	 * handler will be executed when a message of the specified type is received.
+	 *
+	 * This overload is used to register raw class member functions.
+	 *
+	 * @param HandlerType The type of the object handling the messages.
+	 * @param MessageType The type of messages to handle.
+	 * @param Handler The class handling the messages.
+	 * @param HandlerFunc The class function handling the messages.
+	 * @return This instance (for method chaining).
+	 * @see WithHandler
+	 */
+	template<typename HandlerType>
+	FMessageEndpointBuilder& WithCatchall( HandlerType* Handler, typename TRawMessageCatchall<HandlerType>::FuncType HandlerFunc )
+	{
+		Handlers.Add(MakeShareable(new TRawMessageCatchall<HandlerType>(Handler, HandlerFunc)));
+
+		return *this;
+	}
+
+	/**
+	 * Adds a message handler for the given type of messages (via TFunction object).
+	 *
+	 * It is legal to configure multiple handlers for the same message type. Each
+	 * handler will be executed when a message of the specified type is received.
+	 *
+	 * This overload is used to register functions that are compatible with TFunction
+	 * function objects, such as global and static functions, as well as lambdas.
+	 *
+	 * @param MessageType The type of messages to handle.
+	 * @param Function The function object handling the messages.
+	 * @return This instance (for method chaining).
+	 * @see WithHandler
+	 */
+	FMessageEndpointBuilder& WithCatchall( FFunctionMessageCatchall::FuncType HandlerFunc )
+	{
+		Handlers.Add(MakeShareable(new FFunctionMessageCatchall(HandlerFunc)));
+
+		return *this;
+	}
+
 	/**
 	 * Registers a message handler with the endpoint.
 	 *
@@ -155,7 +199,7 @@ public:
 	 *
 	 * @param Handler The handler to add.
 	 * @return This instance (for method chaining).
-	 * @see Handling
+	 * @see Handling, WithCatchall
 	 */
 	FMessageEndpointBuilder& WithHandler( const IMessageHandlerRef& Handler )
 	{
diff --git a/Engine/Source/Runtime/Messaging/Public/Helpers/MessageHandlers.h b/Engine/Source/Runtime/Messaging/Public/Helpers/MessageHandlers.h
--- a/Engine/Source/Runtime/Messaging/Public/Helpers/MessageHandlers.h
+++ b/Engine/Source/Runtime/Messaging/Public/Helpers/MessageHandlers.h
@@ -6,6 +6,57 @@
 #include "IMessageContext.h"
 
 
+/**
+ * Template for catch-all handlers (via raw function pointers).
+ *
+ * @param HandlerType The type of the handler class.
+ */
+template<typename HandlerType>
+class TRawMessageCatchall
+	: public IMessageHandler
+{
+public:
+
+	/** Type definition for function pointers that are compatible with this TRawMessageCatchall. */
+	typedef void (HandlerType::*FuncType)( const TSharedRef<IMessageContext, ESPMode::ThreadSafe>& );
+
+public:
+
+	/**
+	 * Creates and initializes a new message handler.
+	 *
+	 * @param InHandler The object that will handle the messages.
+	 * @param InFunc The object's message handling function.
+	 */
+	TRawMessageCatchall( HandlerType* InHandler, FuncType InFunc )
+		: Handler(InHandler)
+		, Func(InFunc)
+	{
+		check(InHandler != nullptr);
+	}
+
+	/** Virtual destructor. */
+	~TRawMessageCatchall() { }
+
+public:
+
+	// IMessageHandler interface
+	
+	virtual void HandleMessage( const TSharedRef<IMessageContext, ESPMode::ThreadSafe>& Context ) override
+	{
+		(Handler->*Func)(Context);
+	}
+	
+private:
+
+	/** Holds a pointer to the object handling the messages. */
+	HandlerType* Handler;
+
+	/** Holds a pointer to the actual handler function. */
+	FuncType Func;
+};
+
+
 /**
  * Template for handlers of one specific message type (via raw function pointers).
  *
@@ -27,11 +78,11 @@ public:
 	 * Creates and initializes a new message handler.
 	 *
 	 * @param InHandler The object that will handle the messages.
-	 * @param InHandlerFunc The object's message handling function.
+	 * @param InFunc The object's message handling function.
 	 */
-	TRawMessageHandler( HandlerType* InHandler, FuncType InHandlerFunc )
+	TRawMessageHandler( HandlerType* InHandler, FuncType InFunc )
 		: Handler(InHandler)
-		, HandlerFunc(InHandlerFunc)
+		, Func(InFunc)
 	{
 		check(InHandler != nullptr);
 	}
@@ -47,7 +98,7 @@ public:
 	{
 		if (Context->GetMessageType() == MessageType::StaticStruct()->GetFName())
 		{
-			(Handler->*HandlerFunc)(*static_cast<const MessageType*>(Context->GetMessage()), Context);
+			(Handler->*Func)(*static_cast<const MessageType*>(Context->GetMessage()), Context);
 		}	
 	}
 	
@@ -57,7 +108,48 @@ private:
 	HandlerType* Handler;
 
 	/** Holds a pointer to the actual handler function. */
-	FuncType HandlerFunc;
+	FuncType Func;
+};
+
+
+/**
+ * Implements a catch-all handlers (via function objects).
+ */
+class FFunctionMessageCatchall
+	: public IMessageHandler
+{
+public:
+
+	/** Type definition for function objects that are compatible with this TFunctionMessageHandler. */
+	typedef TFunction<void(const TSharedRef<IMessageContext, ESPMode::ThreadSafe>&)> FuncType;
+
+public:
+
+	/**
+	 * Creates and initializes a new message handler.
+	 *
+	 * @param InFunc The object's message handling function.
+	 */
+	FFunctionMessageCatchall( FuncType InFunc )
+		: Func(InFunc)
+	{ }
+
+	/** Virtual destructor. */
+	~FFunctionMessageCatchall() { }
+
+public:
+
+	// IMessageHandler interface
+	
+	virtual void HandleMessage( const TSharedRef<IMessageContext, ESPMode::ThreadSafe>& Context ) override
+	{
+		Func(Context);
+	}
+	
+private:
+
+	/** Holds a pointer to the actual handler function. */
+	FuncType Func;
 };
 
 
@@ -80,11 +172,10 @@ public:
 	/**
 	 * Creates and initializes a new message handler.
 	 *
-	 * @param InHandler The object that will handle the messages.
 	 * @param InHandlerFunc The object's message handling function.
 	 */
-	TFunctionMessageHandler( FuncType InHandlerFunc )
-		: HandlerFunc(InHandlerFunc)
+	TFunctionMessageHandler( FuncType InFunc )
+		: Func(InFunc)
 	{ }
 
 	/** Virtual destructor. */
@@ -98,12 +189,12 @@ public:
 	{
 		if (Context->GetMessageType() == MessageType::StaticStruct()->GetFName())
 		{
-			HandlerFunc(*static_cast<const MessageType*>(Context->GetMessage()), Context);
+			Func(*static_cast<const MessageType*>(Context->GetMessage()), Context);
 		}	
 	}
 	
 private:
 
 	/** Holds a pointer to the actual handler function. */
-	FuncType HandlerFunc;
+	FuncType Func;
 };
