Messaging: Added string to message address conversion

diff --git a/Engine/Source/Runtime/Messaging/Public/IMessageContext.h b/Engine/Source/Runtime/Messaging/Public/IMessageContext.h
--- a/Engine/Source/Runtime/Messaging/Public/IMessageContext.h
+++ b/Engine/Source/Runtime/Messaging/Public/IMessageContext.h
@@ -83,6 +83,7 @@ public:
 	 * Converts this GUID to its string representation.
 	 *
 	 * @return The string representation.
+	 * @see Parse
 	 */
 	FString ToString() const
 	{
@@ -117,6 +118,19 @@ public:
 		return Result;
 	}
 
+	/**
+	 * Converts a string to a message address.
+	 *
+	 * @param String The string to convert.
+	 * @param OutAddress Will contain the parsed address.
+	 * @return true if the string was converted successfully, false otherwise.
+	 * @see ToString
+	 */
+	static bool Parse(const FString& String, FMessageAddress& OutAddress)
+	{
+		return FGuid::Parse(String, OutAddress.UniqueId);
+	}
+
 private:
 
 	/** Holds a unique identifier. */
