Messaging: Fixed handling of local worker thread queues

diff --git a/Engine/Source/Runtime/Core/Private/Async/TaskGraph.cpp b/Engine/Source/Runtime/Core/Private/Async/TaskGraph.cpp
--- a/Engine/Source/Runtime/Core/Private/Async/TaskGraph.cpp
+++ b/Engine/Source/Runtime/Core/Private/Async/TaskGraph.cpp
@@ -984,7 +984,7 @@ public:
 	virtual ENamedThreads::Type GetCurrentThreadIfKnown(bool bLocalQueue) override
 	{
 		ENamedThreads::Type Result = GetCurrentThread();
-		if (Result >= 0 && Result < NumNamedThreads)
+		if (bLocalQueue && Result >= 0 && Result < NumNamedThreads)
 		{
 			Result = ENamedThreads::Type(int32(Result) | int32(ENamedThreads::LocalQueue));
 		}
diff --git a/Engine/Source/Runtime/Messaging/Public/Helpers/MessageEndpoint.h b/Engine/Source/Runtime/Messaging/Public/Helpers/MessageEndpoint.h
--- a/Engine/Source/Runtime/Messaging/Public/Helpers/MessageEndpoint.h
+++ b/Engine/Source/Runtime/Messaging/Public/Helpers/MessageEndpoint.h
@@ -8,6 +8,7 @@
 #include "IMessageHandler.h"
 #include "IReceiveMessages.h"
 #include "ISendMessages.h"
+#include "TaskGraphInterfaces.h"
 
 
 /**
@@ -36,7 +37,7 @@ DECLARE_DELEGATE_RetVal_OneParam(bool, FOnMessageEndpointReceiveMessage, const I
  *
  * By default, messages are received synchronously on the thread that the endpoint was created on.
  * If the message consumer is thread-safe, a more efficient message dispatch can be enabled by calling
- * the SetRecipientThread(ENamedThreads::AnyThread) method.
+ * the SetRecipientThread() method with ENamedThreads::AnyThread.
  *
  * Endpoints that are destroyed or receive messages on non-Game threads should use the static function
  * FMessageEndpoint::SafeRelease() to dispose of the endpoint. This will ensure that there are no race
@@ -77,7 +78,9 @@ public:
 		, Id(FGuid::NewGuid())
 		, InboxEnabled(false)
 		, Name(InName)
-	{ }
+	{
+		SetRecipientThread(FTaskGraphInterface::Get().GetCurrentThreadIfKnown());
+	}
 
 	/** Destructor. */
 	~FMessageEndpoint()
@@ -166,18 +169,7 @@ public:
 	 */
 	void SetRecipientThread( const ENamedThreads::Type& NamedThread )
 	{
-		if (NamedThread == ENamedThreads::GameThread_Local)
-		{
-			RecipientThread = ENamedThreads::GameThread;
-		}
-		else if (NamedThread == ENamedThreads::RenderThread_Local)
-		{
-			RecipientThread = ENamedThreads::RenderThread;
-		}
-		else
-		{
-			RecipientThread = NamedThread;
-		}		
+		RecipientThread = ENamedThreads::GetThreadIndex(NamedThread);
 	}
 
 public:
