MessagingDebugger: Simplified message tracer API

diff --git a/Engine/Plugins/Messaging/MessagingDebugger/Source/MessagingDebugger/Private/Widgets/SMessagingDebugger.cpp b/Engine/Plugins/Messaging/MessagingDebugger/Source/MessagingDebugger/Private/Widgets/SMessagingDebugger.cpp
--- a/Engine/Plugins/Messaging/MessagingDebugger/Source/MessagingDebugger/Private/Widgets/SMessagingDebugger.cpp
+++ b/Engine/Plugins/Messaging/MessagingDebugger/Source/MessagingDebugger/Private/Widgets/SMessagingDebugger.cpp
@@ -27,7 +27,7 @@ static const FName ToolbarTabId("Toolbar");
 
 SMessagingDebugger::SMessagingDebugger()
 	: CommandList(MakeShareable(new FUICommandList))
-	, MessageTracer(NULL)
+	, MessageTracer(nullptr)
 	, Model(MakeShareable(new FMessagingDebuggerModel()))
 { }
 
@@ -262,7 +262,7 @@ void SMessagingDebugger::HandleClearHistoryCommandExecute()
 
 bool SMessagingDebugger::HandleContinueDebuggerCommandCanExecute() const
 {
-	return MessageTracer->IsBreaking();
+	return (!MessageTracer->IsRunning() || MessageTracer->IsBreaking());
 }
 
 
@@ -292,7 +292,7 @@ bool SMessagingDebugger::HandleStartDebuggerCommandCanExecute() const
 
 void SMessagingDebugger::HandleStartDebuggerCommandExecute()
 {
-	MessageTracer->Start();
+	MessageTracer->Continue();
 }
 
 
diff --git a/Engine/Source/Runtime/Messaging/Private/Bus/MessageTracer.h b/Engine/Source/Runtime/Messaging/Private/Bus/MessageTracer.h
--- a/Engine/Source/Runtime/Messaging/Private/Bus/MessageTracer.h
+++ b/Engine/Source/Runtime/Messaging/Private/Bus/MessageTracer.h
@@ -32,7 +32,7 @@ public:
 	 * @param Interceptor The added interceptor.
 	 * @param MessageType The type of messages being intercepted.
 	 */
-	void TraceAddedInterceptor( const IMessageInterceptorRef& Interceptor, const FName& MessageType );
+	void TraceAddedInterceptor(const IMessageInterceptorRef& Interceptor, const FName& MessageType);
 
 	/**
 	 * Notifies the tracer that a message recipient has been added to the message bus.
@@ -40,14 +40,14 @@ public:
 	 * @param Address The address of the added recipient.
 	 * @param Recipient The added recipient.
 	 */
-	void TraceAddedRecipient( const FMessageAddress& Address, const IReceiveMessagesRef& Recipient );
+	void TraceAddedRecipient(const FMessageAddress& Address, const IReceiveMessagesRef& Recipient);
 
 	/**
 	 * Notifies the tracer that a message subscription has been added to the message bus.
 	 *
 	 * @param Subscription The added subscription.
 	 */
-	void TraceAddedSubscription( const IMessageSubscriptionRef& Subscription );
+	void TraceAddedSubscription(const IMessageSubscriptionRef& Subscription);
 
 	/**
 	 * Notifies the tracer that a message has been dispatched.
@@ -56,7 +56,7 @@ public:
 	 * @param Recipient The message recipient.
 	 * @param Async Whether the message was dispatched asynchronously.
 	 */
-	void TraceDispatchedMessage( const IMessageContextRef& Context, const IReceiveMessagesRef& Recipient, bool Async );
+	void TraceDispatchedMessage(const IMessageContextRef& Context, const IReceiveMessagesRef& Recipient, bool Async);
 
 	/**
 	 * Notifies the tracer that a message has been handled.
@@ -64,7 +64,7 @@ public:
 	 * @param Context The context of the dispatched message.
 	 * @param Recipient The message recipient that handled the message.
 	 */
-	void TraceHandledMessage( const IMessageContextRef& Context, const IReceiveMessagesRef& Recipient );
+	void TraceHandledMessage(const IMessageContextRef& Context, const IReceiveMessagesRef& Recipient);
 
 	/**
 	 * Notifies the tracer that a message has been intercepted.
@@ -72,7 +72,7 @@ public:
 	 * @param Context The context of the intercepted message.
 	 * @param Interceptor The interceptor.
 	 */
-	void TraceInterceptedMessage( const IMessageContextRef& Context, const IMessageInterceptorRef& Interceptor );
+	void TraceInterceptedMessage(const IMessageContextRef& Context, const IMessageInterceptorRef& Interceptor);
 
 	/**
 	 * Notifies the tracer that a message interceptor has been removed from the message bus.
@@ -80,14 +80,14 @@ public:
 	 * @param Interceptor The removed interceptor.
 	 * @param MessageType The type of messages that is no longer being intercepted.
 	 */
-	void TraceRemovedInterceptor( const IMessageInterceptorRef& Interceptor, const FName& MessageType );
+	void TraceRemovedInterceptor(const IMessageInterceptorRef& Interceptor, const FName& MessageType);
 
 	/**
 	 * Notifies the tracer that a recipient has been removed from the message bus.
 	 *
 	 * @param Address The address of the removed recipient.
 	 */
-	void TraceRemovedRecipient( const FMessageAddress& Address );
+	void TraceRemovedRecipient(const FMessageAddress& Address);
 
 	/**
 	 * Notifies the tracer that a message subscription has been removed from the message bus.
@@ -95,21 +95,21 @@ public:
 	 * @param Subscriber The removed subscriber.
 	 * @param MessageType The type of messages no longer being subscribed to.
 	 */
-	void TraceRemovedSubscription( const IMessageSubscriptionRef& Subscription, const FName& MessageType );
+	void TraceRemovedSubscription(const IMessageSubscriptionRef& Subscription, const FName& MessageType);
 
 	/**
 	 * Notifies the tracer that a message has been routed.
 	 *
 	 * @param Context The context of the routed message.
 	 */
-	void TraceRoutedMessage( const IMessageContextRef& Context );
+	void TraceRoutedMessage(const IMessageContextRef& Context);
 
 	/**
 	 * Notifies the tracer that a message has been sent.
 	 *
 	 * @param Context The context of the sent message.
 	 */
-	void TraceSentMessage( const IMessageContextRef& Context );
+	void TraceSentMessage(const IMessageContextRef& Context);
 
 public:
 
@@ -122,18 +122,20 @@ public:
 
 	virtual void Continue() override
 	{
-		if (!Breaking)
+		if (!Running)
 		{
-			return;
+			Running = true;
+		}
+		else if (Breaking)
+		{
+			Breaking = false;
+			ContinueEvent->Trigger();
 		}
-
-		Breaking = false;
-		ContinueEvent->Trigger();
 	}
 
-	virtual int32 GetEndpoints( TArray<FMessageTracerEndpointInfoPtr>& OutEndpoints ) const override;
-	virtual int32 GetMessages( TArray<FMessageTracerMessageInfoPtr>& OutMessages ) const override;
-	virtual int32 GetMessageTypes( TArray<FMessageTracerTypeInfoPtr>& OutTypes ) const override;
+	virtual int32 GetEndpoints(TArray<FMessageTracerEndpointInfoPtr>& OutEndpoints) const override;
+	virtual int32 GetMessages(TArray<FMessageTracerMessageInfoPtr>& OutMessages) const override;
+	virtual int32 GetMessageTypes(TArray<FMessageTracerTypeInfoPtr>& OutTypes) const override;
 
 	virtual bool HasMessages() const override
 	{
@@ -173,16 +175,6 @@ public:
 		ResetPending = true;
 	}
 
-	virtual void Start() override
-	{
-		if (Running)
-		{
-			return;
-		}
-
-		Running = true;
-	}
-
 	virtual void Step() override
 	{
 		if (!Breaking)
@@ -209,7 +201,7 @@ public:
 		}
 	}
 
-	virtual bool Tick( float DeltaTime ) override;
+	virtual bool Tick(float DeltaTime) override;
 
 protected:
 
@@ -223,7 +215,7 @@ protected:
 	 *
 	 * @param Context The context of the message to consider for breaking.
 	 */
-	bool ShouldBreak( const IMessageContextRef& Context ) const;
+	bool ShouldBreak(const IMessageContextRef& Context) const;
 
 private:
 
diff --git a/Engine/Source/Runtime/Messaging/Public/IMessageTracer.h b/Engine/Source/Runtime/Messaging/Public/IMessageTracer.h
--- a/Engine/Source/Runtime/Messaging/Public/IMessageTracer.h
+++ b/Engine/Source/Runtime/Messaging/Public/IMessageTracer.h
@@ -173,7 +173,7 @@ public:
 	virtual void Break() = 0;
 
 	/**
-	 * Continues message routing from the current breakpoint.
+	 * Starts the tracer or continues message routing from the current breakpoint.
 	 *
 	 * @see Break, Step
 	 */
@@ -183,7 +183,7 @@ public:
 	 * Checks whether the tracer is currently at a breakpoint.
 	 *
 	 * @return true if at breakpoint, false otherwise.
-	 * @see Break
+	 * @see Break, IsRunning
 	 */
 	virtual bool IsBreaking() const = 0;
 
@@ -191,32 +191,24 @@ public:
 	 * Checks whether the tracer is currently running.
 	 *
 	 * @return true if the tracer is running, false otherwise.
-	 *
-	 * @see Start, Stop
+	 * @see Continue, IsBreaking, Stop
 	 */
 	virtual bool IsRunning() const = 0;
 
 	/** Resets the tracer. */
 	virtual void Reset() = 0;
 
-	/**
-	 * Starts the tracer.
-	 *
-	 * @see IsRunning, Stop
-	 */
-	virtual void Start() = 0;
-
 	/**
 	 * Steps the tracer to the next message.
 	 *
-	 * @see Break, Continue
+	 * @see Break, Continue, Stop
 	 */
 	virtual void Step() = 0;
 
 	/**
 	 * Stops the tracer.
 	 *
-	 * @see IsRunning, Start
+	 * @see Continue, IsRunning, Step
 	 */
 	virtual void Stop() = 0;
 
@@ -226,7 +218,7 @@ public:
 	 * @param DeltaTime The time in seconds since the last tick.
 	 * @return true if any events were processed.
 	 */
-	virtual bool Tick( float DeltaTime ) = 0;
+	virtual bool Tick(float DeltaTime) = 0;
 
 public:
 
@@ -236,14 +228,14 @@ public:
 	 * @param OutEndpoints Will contain the list of endpoints.
 	 * @return The number of endpoints returned.
 	 */
-	virtual int32 GetEndpoints( TArray<FMessageTracerEndpointInfoPtr>& OutEndpoints ) const = 0;
+	virtual int32 GetEndpoints(TArray<FMessageTracerEndpointInfoPtr>& OutEndpoints) const = 0;
 
 	/**
 	 * Gets the collection of known messages.
 	 *
 	 * @return The messages.
 	 */
-	virtual int32 GetMessages( TArray<FMessageTracerMessageInfoPtr>& OutMessages ) const = 0;
+	virtual int32 GetMessages(TArray<FMessageTracerMessageInfoPtr>& OutMessages) const = 0;
 
 	/**
 	 * Gets the list of known message types filtered by name.
@@ -252,7 +244,7 @@ public:
 	 * @param OutTypes Will contain the list of message types.
 	 * @return The number of message types returned.
 	 */
-	virtual int32 GetMessageTypes( TArray<FMessageTracerTypeInfoPtr>& OutTypes ) const = 0;
+	virtual int32 GetMessageTypes(TArray<FMessageTracerTypeInfoPtr>& OutTypes) const = 0;
 
 	/**
 	 * Checks whether there are any messages in the history.
