Messaging: Implemented interceptor tracing

diff --git a/Engine/Source/Runtime/Messaging/Private/Bus/MessageTracer.cpp b/Engine/Source/Runtime/Messaging/Private/Bus/MessageTracer.cpp
--- a/Engine/Source/Runtime/Messaging/Private/Bus/MessageTracer.cpp
+++ b/Engine/Source/Runtime/Messaging/Private/Bus/MessageTracer.cpp
@@ -32,7 +32,18 @@ void FMessageTracer::TraceAddedInterceptor( const IMessageInterceptorRef& Interc
 	double Timestamp = FPlatformTime::Seconds();
 
 	Traces.Enqueue([=]() {
-		// @todo gmp: trace added message interceptors
+		// create interceptor information
+		FMessageTracerInterceptorInfoPtr& InterceptorInfo = Interceptors.FindOrAdd(Interceptor->GetInterceptorId());
+
+		if (!InterceptorInfo.IsValid())
+		{
+			InterceptorInfo = MakeShareable(new FMessageTracerInterceptorInfo());
+		}
+
+		// initialize interceptor information
+		InterceptorInfo->Name = Interceptor->GetDebugName();
+		InterceptorInfo->TimeRegistered = Timestamp;
+		InterceptorInfo->TimeUnregistered = 0;
 	});
 }
 
@@ -170,33 +181,49 @@ void FMessageTracer::TraceInterceptedMessage( const IMessageContextRef& Context,
 	double Timestamp = FPlatformTime::Seconds();
 
 	Traces.Enqueue([=]() {
-		// @todo gmp: trace intercepted messages
+		// look up message & interceptor info
+		FMessageTracerMessageInfoPtr MessageInfo = MessageInfos.FindRef(Context);
+
+		if (!MessageInfo.IsValid())
+		{
+			return;
+		}
+
+		MessageInfo->Intercepted = true;
+
+		FMessageTracerInterceptorInfoPtr InterceptorInfo = Interceptors.FindRef(Interceptor->GetInterceptorId());
+
+		if (!InterceptorInfo.IsValid())
+		{
+			return;
+		}
+
+		// update interceptor information
+		InterceptorInfo->InterceptedMessages.Add(MessageInfo);
 	});		
 }
 
 
 void FMessageTracer::TraceRemovedInterceptor( const IMessageInterceptorRef& Interceptor, const FName& MessageType )
 {
-	if (!Running)
-	{
-		return;
-	}
-
 	double Timestamp = FPlatformTime::Seconds();
 
 	Traces.Enqueue([=]() {
-		// @todo gmp: trace removed message interceptors
+		FMessageTracerInterceptorInfoPtr InterceptorInfo = Interceptors.FindRef(Interceptor->GetInterceptorId());
+
+		if (!InterceptorInfo.IsValid())
+		{
+			return;
+		}
+
+		// update interceptor information
+		InterceptorInfo->TimeUnregistered = Timestamp;
 	});
 }
 
 
 void FMessageTracer::TraceRemovedRecipient( const FMessageAddress& Address )
 {
-	if (!Running)
-	{
-		return;
-	}
-
 	double Timestamp = FPlatformTime::Seconds();
 
 	Traces.Enqueue([=]() {
@@ -282,6 +309,7 @@ void FMessageTracer::TraceSentMessage( const IMessageContextRef& Context )
 		FMessageTracerMessageInfoRef MessageInfo = MakeShareable(new FMessageTracerMessageInfo());
 	
 		MessageInfo->Context = Context;
+		MessageInfo->Intercepted = false;
 		MessageInfo->SenderInfo = EndpointInfo;
 		MessageInfo->TimeRouted = 0.0;
 		MessageInfo->TimeSent = Timestamp;
diff --git a/Engine/Source/Runtime/Messaging/Private/Bus/MessageTracer.h b/Engine/Source/Runtime/Messaging/Private/Bus/MessageTracer.h
--- a/Engine/Source/Runtime/Messaging/Private/Bus/MessageTracer.h
+++ b/Engine/Source/Runtime/Messaging/Private/Bus/MessageTracer.h
@@ -234,6 +234,9 @@ private:
 	/** Holds an event signaling that messaging routing can continue. */
 	FEvent* ContinueEvent;
 
+	/** The collection of known interceptors. */
+	TMap<FGuid, FMessageTracerInterceptorInfoPtr> Interceptors;
+
 	/** Holds the collection of endpoints for known recipient identifiers. */
 	TMap<FGuid, FMessageTracerEndpointInfoPtr> RecipientsToEndpointInfos;
 
diff --git a/Engine/Source/Runtime/Messaging/Public/IMessageInterceptor.h b/Engine/Source/Runtime/Messaging/Public/IMessageInterceptor.h
--- a/Engine/Source/Runtime/Messaging/Public/IMessageInterceptor.h
+++ b/Engine/Source/Runtime/Messaging/Public/IMessageInterceptor.h
@@ -13,13 +13,29 @@ class IMessageInterceptor
 {
 public:
 
+	/**
+	 * Gets the interceptor's name (for debugging purposes).
+	 *
+	 * @return The debug name.
+	 * @see GetInterceptorId
+	 */
+	virtual FName GetDebugName() const = 0;
+
+	/**
+	 * Gets the interceptor's unique identifier (for debugging purposes).
+	 *
+	 * @return The interceptor's identifier.
+	 * @see GetDebugName
+	 */
+	virtual const FGuid& GetInterceptorId() const = 0;
+
 	/**
 	 * Intercepts a message before it is being passed to the message router.
 	 *
 	 * @param Context The context of the message to intercept.
 	 * @return true if the message was intercepted and should not be routed, false otherwise.
 	 */
-	virtual bool InterceptMessage( const TSharedRef<IMessageContext, ESPMode::ThreadSafe>& Context ) = 0;
+	virtual bool InterceptMessage(const TSharedRef<IMessageContext, ESPMode::ThreadSafe>& Context) = 0;
 
 public:
 
diff --git a/Engine/Source/Runtime/Messaging/Public/IMessageTracer.h b/Engine/Source/Runtime/Messaging/Public/IMessageTracer.h
--- a/Engine/Source/Runtime/Messaging/Public/IMessageTracer.h
+++ b/Engine/Source/Runtime/Messaging/Public/IMessageTracer.h
@@ -24,6 +24,12 @@ typedef TSharedPtr<struct FMessageTracerEndpointInfo> FMessageTracerEndpointInfo
 /** Type definition for shared references to instances of FMessageTracerEndpointInfo. */
 typedef TSharedRef<struct FMessageTracerEndpointInfo> FMessageTracerEndpointInfoRef;
 
+/** Type definition for shared pointers to instances of FMessageTracerInterceptorInfo. */
+typedef TSharedPtr<struct FMessageTracerInterceptorInfo> FMessageTracerInterceptorInfoPtr;
+
+/** Type definition for shared references to instances of FMessageTracerInterceptorInfo. */
+typedef TSharedRef<struct FMessageTracerInterceptorInfo> FMessageTracerInterceptorInfoRef;
+
 /** Type definition for shared pointers to instances of FMessageTracerMessageInfo. */
 typedef TSharedPtr<struct FMessageTracerMessageInfo> FMessageTracerMessageInfoPtr;
 
@@ -116,6 +122,25 @@ struct FMessageTracerEndpointInfo
 };
 
 
+/**
+ * Structure for message interceptor debug information.
+ */
+struct FMessageTracerInterceptorInfo
+{
+	/** Holds the interceptor's human readable name. */
+	FName Name;
+
+	/** Holds the list of messages intercepted by this interceptor. */
+	TArray<FMessageTracerMessageInfoPtr> InterceptedMessages;
+
+	/** Holds the time at which this interceptor was registered. */
+	double TimeRegistered;
+
+	/** Holds the time at which this interceptor was unregistered. */
+	double TimeUnregistered;
+};
+
+
 /**
  * Structure for message debug information.
  */
@@ -127,6 +152,9 @@ struct FMessageTracerMessageInfo
 	/** Holds the message's dispatch states per endpoint. */
 	TMap<FMessageTracerEndpointInfoPtr, FMessageTracerDispatchStatePtr> DispatchStates;
 
+	/** Whether the message was intercepted. */
+	bool Intercepted;
+
 	/** Pointer to the sender's endpoint information. */
 	FMessageTracerEndpointInfoPtr SenderInfo;
 
