Messaging: Modernization pass (ranged for-loops, reduced nesting, fixed virtual destructors)

diff --git a/Engine/Source/Runtime/Messaging/Private/Bridge/MessageAddressBook.h b/Engine/Source/Runtime/Messaging/Private/Bridge/MessageAddressBook.h
--- a/Engine/Source/Runtime/Messaging/Private/Bridge/MessageAddressBook.h
+++ b/Engine/Source/Runtime/Messaging/Private/Bridge/MessageAddressBook.h
@@ -30,7 +30,7 @@ public:
 	 * @Address The message address to add.
 	 * @NodeId The identifier of the remote node that handles the message address.
 	 */
-	void Add( const FMessageAddress& Address, const FGuid& NodeId )
+	void Add(const FMessageAddress& Address, const FGuid& NodeId)
 	{
 		FScopeLock Lock(CriticalSection);
 
@@ -51,7 +51,7 @@ public:
 	 * @param Address The address to check.
 	 * @return true if the address is known, false otherwise.
 	 */
-	bool Contains( const FMessageAddress& Address )
+	bool Contains(const FMessageAddress& Address)
 	{
 		FScopeLock Lock(CriticalSection);
 
@@ -64,15 +64,15 @@ public:
 	 * @param Addresses The address list to retrieve the node identifiers for.
 	 * @return The list of node identifiers.
 	 */
-	TArray<FGuid> GetNodesFor( const TArray<FMessageAddress>& Addresses )
+	TArray<FGuid> GetNodesFor(const TArray<FMessageAddress>& Addresses)
 	{
 		TArray<FGuid> FoundNodes;
 
 		FScopeLock Lock(CriticalSection);
 
-		for (int32 AddressIndex = 0; AddressIndex < Addresses.Num(); ++AddressIndex)
+		for (const auto& Address : Addresses)
 		{
-			FGuid* NodeId = Entries.Find(Addresses[AddressIndex]);
+			FGuid* NodeId = Entries.Find(Address);
 
 			if (NodeId != nullptr)
 			{
@@ -92,7 +92,7 @@ public:
 	 * @param OutRemovedRecipients Will hold a list of recipients that were removed.
 	 * @see Clear, RemoveNode
 	 */
-	void RemoveAll( TArray<FMessageAddress>& OutRemovedAddresses )
+	void RemoveAll(TArray<FMessageAddress>& OutRemovedAddresses)
 	{
 		OutRemovedAddresses.Reset();
 
@@ -109,23 +109,23 @@ public:
 	 * @param OutRemovedRecipients Will hold a list of recipients that were removed.
 	 * @see Clear, RemoveAllNodes
 	 */
-	void RemoveNode( const FGuid& NodeId, TArray<FMessageAddress>& OutRemovedAddresses )
+	void RemoveNode(const FGuid& NodeId, TArray<FMessageAddress>& OutRemovedAddresses)
 	{
 		OutRemovedAddresses.Reset();
 
 		FScopeLock Lock(CriticalSection);
 
-		for (TMap<FMessageAddress, FGuid>::TConstIterator It(Entries); It; ++It)
+		for (const auto& EntryPair : Entries)
 		{
-			if (It.Value() == NodeId)
+			if (EntryPair.Value == NodeId)
 			{
-				OutRemovedAddresses.Add(It.Key());
+				OutRemovedAddresses.Add(EntryPair.Key);
 			}
 		}
 
-		for (int32 Index = 0; Index < OutRemovedAddresses.Num(); ++Index)
+		for (const auto& Address : OutRemovedAddresses)
 		{
-			Entries.Remove(OutRemovedAddresses[Index]);
+			Entries.Remove(Address);
 		}
 	}
 
diff --git a/Engine/Source/Runtime/Messaging/Private/Bridge/MessageBridge.cpp b/Engine/Source/Runtime/Messaging/Private/Bridge/MessageBridge.cpp
--- a/Engine/Source/Runtime/Messaging/Private/Bridge/MessageBridge.cpp
+++ b/Engine/Source/Runtime/Messaging/Private/Bridge/MessageBridge.cpp
@@ -6,7 +6,7 @@
 /* FMessageBridge structors
  *****************************************************************************/
 
-FMessageBridge::FMessageBridge( const FMessageAddress InAddress, const IMessageBusRef& InBus, const IMessageTransportRef& InTransport )
+FMessageBridge::FMessageBridge(const FMessageAddress InAddress, const IMessageBusRef& InBus, const IMessageTransportRef& InTransport)
 	: Address(InAddress)
 	, Bus(InBus)
 	, Enabled(false)
@@ -94,7 +94,7 @@ void FMessageBridge::Enable()
 /* IReceiveMessages interface
  *****************************************************************************/
 
-void FMessageBridge::ReceiveMessage( const IMessageContextRef& Context )
+void FMessageBridge::ReceiveMessage(const IMessageContextRef& Context)
 {
 	if (!Enabled)
 	{
@@ -122,7 +122,7 @@ void FMessageBridge::ReceiveMessage( const IMessageContextRef& Context )
 /* ISendMessages interface
  *****************************************************************************/
 
-void FMessageBridge::NotifyMessageError( const IMessageContextRef& Context, const FString& Error )
+void FMessageBridge::NotifyMessageError(const IMessageContextRef& Context, const FString& Error)
 {
 	// deprecated
 }
@@ -153,7 +153,7 @@ void FMessageBridge::HandleMessageBusShutdown()
 }
 
 
-void FMessageBridge::HandleTransportMessageReceived( const IMessageContextRef& Context, const FGuid& NodeId )
+void FMessageBridge::HandleTransportMessageReceived(const IMessageContextRef& Context, const FGuid& NodeId)
 {
 	if (!Enabled || !Bus.IsValid())
 	{
@@ -178,7 +178,7 @@ void FMessageBridge::HandleTransportMessageReceived( const IMessageContextRef& C
 }
 
 
-void FMessageBridge::HandleTransportNodeLost( const FGuid& LostNodeId )
+void FMessageBridge::HandleTransportNodeLost(const FGuid& LostNodeId)
 {
 	TArray<FMessageAddress> RemovedAddresses;
 
@@ -188,9 +188,9 @@ void FMessageBridge::HandleTransportNodeLost( const FGuid& LostNodeId )
 	// unregister endpoints
 	if (Bus.IsValid())
 	{
-		for (int32 AddressIndex = 0; AddressIndex < RemovedAddresses.Num(); ++AddressIndex)
+		for (const auto& Address : RemovedAddresses)
 		{
-			Bus->Unregister(RemovedAddresses[AddressIndex]);
+			Bus->Unregister(Address);
 		}
 	}
 }
diff --git a/Engine/Source/Runtime/Messaging/Private/Bridge/MessageBridge.h b/Engine/Source/Runtime/Messaging/Private/Bridge/MessageBridge.h
--- a/Engine/Source/Runtime/Messaging/Private/Bridge/MessageBridge.h
+++ b/Engine/Source/Runtime/Messaging/Private/Bridge/MessageBridge.h
@@ -38,10 +38,10 @@ public:
 	 * @param InBus The message bus that this node is connected to.
 	 * @param InTransport The transport mechanism to use.
 	 */
-	FMessageBridge( const FMessageAddress InAddress, const IMessageBusRef& InBus, const IMessageTransportRef& InTransport );
+	FMessageBridge(const FMessageAddress InAddress, const IMessageBusRef& InBus, const IMessageTransportRef& InTransport);
 
-	/** Destructor. */
-	~FMessageBridge();
+	/** Virtual destructor. */
+	virtual ~FMessageBridge();
 
 public:
 
@@ -79,7 +79,7 @@ public:
 		return false;
 	}
 
-	virtual void ReceiveMessage( const IMessageContextRef& Context ) override;
+	virtual void ReceiveMessage(const IMessageContextRef& Context) override;
 
 public:
 
@@ -90,7 +90,7 @@ public:
 		return Address;
 	}
 
-	virtual void NotifyMessageError( const IMessageContextRef& Context, const FString& Error ) override;
+	virtual void NotifyMessageError(const IMessageContextRef& Context, const FString& Error) override;
 
 protected:
 
@@ -103,10 +103,10 @@ private:
 	void HandleMessageBusShutdown();
 
 	/** Callback for messages received from the transport layer. */
-	void HandleTransportMessageReceived( const IMessageContextRef& Envelope, const FGuid& NodeId );
+	void HandleTransportMessageReceived(const IMessageContextRef& Envelope, const FGuid& NodeId);
 
 	/** Callback for lost remote nodes. */
-	void HandleTransportNodeLost( const FGuid& LostNodeId );
+	void HandleTransportNodeLost(const FGuid& LostNodeId);
 
 private:
 
diff --git a/Engine/Source/Runtime/Messaging/Private/Bus/MessageBus.cpp b/Engine/Source/Runtime/Messaging/Private/Bus/MessageBus.cpp
--- a/Engine/Source/Runtime/Messaging/Private/Bus/MessageBus.cpp
+++ b/Engine/Source/Runtime/Messaging/Private/Bus/MessageBus.cpp
@@ -6,7 +6,7 @@
 /* FMessageBus structors
  *****************************************************************************/
 
-FMessageBus::FMessageBus( const IAuthorizeMessageRecipientsPtr& InRecipientAuthorizer )
+FMessageBus::FMessageBus(const IAuthorizeMessageRecipientsPtr& InRecipientAuthorizer)
 	: RecipientAuthorizer(InRecipientAuthorizer)
 {
 	Router = new FMessageRouter();
@@ -16,7 +16,7 @@ FMessageBus::FMessageBus( const IAuthorizeMessageRecipientsPtr& InRecipientAutho
 }
 
 
-FMessageBus::~FMessageBus( )
+FMessageBus::~FMessageBus()
 {
 	Shutdown();
 
@@ -27,7 +27,7 @@ FMessageBus::~FMessageBus( )
 /* IMessageBus interface
  *****************************************************************************/
 
-void FMessageBus::Forward( const IMessageContextRef& Context, const TArray<FMessageAddress>& Recipients, const FTimespan& Delay, const ISendMessagesRef& Forwarder )
+void FMessageBus::Forward(const IMessageContextRef& Context, const TArray<FMessageAddress>& Recipients, const FTimespan& Delay, const ISendMessagesRef& Forwarder)
 {
 	Router->RouteMessage(MakeShareable(new FMessageContext(Context, Forwarder->GetSenderAddress(), Recipients, EMessageScope::Process, FDateTime::UtcNow() + Delay, FTaskGraphInterface::Get().GetCurrentThreadIfKnown())));
 }
@@ -39,15 +39,17 @@ IMessageTracerRef FMessageBus::GetTracer()
 }
 
 
-void FMessageBus::Intercept( const IMessageInterceptorRef& Interceptor, const FName& MessageType )
+void FMessageBus::Intercept(const IMessageInterceptorRef& Interceptor, const FName& MessageType)
 {
-	if (MessageType != NAME_None)
+	if (MessageType == NAME_None)
 	{
-		if (!RecipientAuthorizer.IsValid() || RecipientAuthorizer->AuthorizeInterceptor(Interceptor, MessageType))
-		{
-			Router->AddInterceptor(Interceptor, MessageType);
-		}			
+		return;
 	}
+
+	if (!RecipientAuthorizer.IsValid() || RecipientAuthorizer->AuthorizeInterceptor(Interceptor, MessageType))
+	{
+		Router->AddInterceptor(Interceptor, MessageType);
+	}			
 }
 
 
@@ -57,25 +59,25 @@ FOnMessageBusShutdown& FMessageBus::OnShutdown()
 }
 
 
-void FMessageBus::Publish( void* Message, UScriptStruct* TypeInfo, EMessageScope Scope, const FTimespan& Delay, const FDateTime& Expiration, const ISendMessagesRef& Publisher )
+void FMessageBus::Publish(void* Message, UScriptStruct* TypeInfo, EMessageScope Scope, const FTimespan& Delay, const FDateTime& Expiration, const ISendMessagesRef& Publisher)
 {
 	Router->RouteMessage(MakeShareable(new FMessageContext(Message, TypeInfo, nullptr, Publisher->GetSenderAddress(), TArray<FMessageAddress>(), Scope, FDateTime::UtcNow() + Delay, Expiration, FTaskGraphInterface::Get().GetCurrentThreadIfKnown())));
 }
 
 
-void FMessageBus::Register( const FMessageAddress& Address, const IReceiveMessagesRef& Recipient )
+void FMessageBus::Register(const FMessageAddress& Address, const IReceiveMessagesRef& Recipient)
 {
 	Router->AddRecipient(Address, Recipient);
 }
 
 
-void FMessageBus::Send( void* Message, UScriptStruct* TypeInfo, const IMessageAttachmentPtr& Attachment, const TArray<FMessageAddress>& Recipients, const FTimespan& Delay, const FDateTime& Expiration, const ISendMessagesRef& Sender )
+void FMessageBus::Send(void* Message, UScriptStruct* TypeInfo, const IMessageAttachmentPtr& Attachment, const TArray<FMessageAddress>& Recipients, const FTimespan& Delay, const FDateTime& Expiration, const ISendMessagesRef& Sender)
 {
 	Router->RouteMessage(MakeShareable(new FMessageContext(Message, TypeInfo, Attachment, Sender->GetSenderAddress(), Recipients, EMessageScope::Network, FDateTime::UtcNow() + Delay, Expiration, FTaskGraphInterface::Get().GetCurrentThreadIfKnown())));
 }
 
 
-void FMessageBus::Shutdown( )
+void FMessageBus::Shutdown()
 {
 	if (RouterThread != nullptr)
 	{
@@ -88,7 +90,7 @@ void FMessageBus::Shutdown( )
 }
 
 
-IMessageSubscriptionPtr FMessageBus::Subscribe( const IReceiveMessagesRef& Subscriber, const FName& MessageType, const FMessageScopeRange& ScopeRange )
+IMessageSubscriptionPtr FMessageBus::Subscribe(const IReceiveMessagesRef& Subscriber, const FName& MessageType, const FMessageScopeRange& ScopeRange)
 {
 	if (MessageType != NAME_None)
 	{
@@ -105,7 +107,7 @@ IMessageSubscriptionPtr FMessageBus::Subscribe( const IReceiveMessagesRef& Subsc
 }
 
 
-void FMessageBus::Unintercept( const IMessageInterceptorRef& Interceptor, const FName& MessageType )
+void FMessageBus::Unintercept(const IMessageInterceptorRef& Interceptor, const FName& MessageType)
 {
 	if (MessageType != NAME_None)
 	{
@@ -114,7 +116,7 @@ void FMessageBus::Unintercept( const IMessageInterceptorRef& Interceptor, const
 }
 
 
-void FMessageBus::Unregister( const FMessageAddress& Address )
+void FMessageBus::Unregister(const FMessageAddress& Address)
 {
 	if (!RecipientAuthorizer.IsValid() || RecipientAuthorizer->AuthorizeUnregistration(Address))
 	{
@@ -123,7 +125,7 @@ void FMessageBus::Unregister( const FMessageAddress& Address )
 }
 
 
-void FMessageBus::Unsubscribe( const IReceiveMessagesRef& Subscriber, const FName& MessageType )
+void FMessageBus::Unsubscribe(const IReceiveMessagesRef& Subscriber, const FName& MessageType)
 {
 	if (MessageType != NAME_None)
 	{
diff --git a/Engine/Source/Runtime/Messaging/Private/Bus/MessageBus.h b/Engine/Source/Runtime/Messaging/Private/Bus/MessageBus.h
--- a/Engine/Source/Runtime/Messaging/Private/Bus/MessageBus.h
+++ b/Engine/Source/Runtime/Messaging/Private/Bus/MessageBus.h
@@ -31,27 +31,27 @@ public:
 	 *
 	 * @param InRecipientAuthorizer An optional recipient authorizer.
 	 */
-	FMessageBus( const IAuthorizeMessageRecipientsPtr& InRecipientAuthorizer );
+	FMessageBus(const IAuthorizeMessageRecipientsPtr& InRecipientAuthorizer);
 
-	/** Destructor. */
-	~FMessageBus();
+	/** Virtual destructor. */
+	virtual ~FMessageBus();
 
 public:
 
 	// IMessageBus interface
 
-	virtual void Forward( const IMessageContextRef& Context, const TArray<FMessageAddress>& Recipients, const FTimespan& Delay, const ISendMessagesRef& Forwarder ) override;
+	virtual void Forward(const IMessageContextRef& Context, const TArray<FMessageAddress>& Recipients, const FTimespan& Delay, const ISendMessagesRef& Forwarder) override;
 	virtual IMessageTracerRef GetTracer() override;
-	virtual void Intercept( const IMessageInterceptorRef& Interceptor, const FName& MessageType ) override;
+	virtual void Intercept(const IMessageInterceptorRef& Interceptor, const FName& MessageType) override;
 	virtual FOnMessageBusShutdown& OnShutdown() override;
-	virtual void Publish( void* Message, UScriptStruct* TypeInfo, EMessageScope Scope, const FTimespan& Delay, const FDateTime& Expiration, const ISendMessagesRef& Publisher ) override;
-	virtual void Register( const FMessageAddress& Address, const IReceiveMessagesRef& Recipient ) override;
-	virtual void Send( void* Message, UScriptStruct* TypeInfo, const IMessageAttachmentPtr& Attachment, const TArray<FMessageAddress>& Recipients, const FTimespan& Delay, const FDateTime& Expiration, const ISendMessagesRef& Sender ) override;
+	virtual void Publish(void* Message, UScriptStruct* TypeInfo, EMessageScope Scope, const FTimespan& Delay, const FDateTime& Expiration, const ISendMessagesRef& Publisher) override;
+	virtual void Register(const FMessageAddress& Address, const IReceiveMessagesRef& Recipient) override;
+	virtual void Send(void* Message, UScriptStruct* TypeInfo, const IMessageAttachmentPtr& Attachment, const TArray<FMessageAddress>& Recipients, const FTimespan& Delay, const FDateTime& Expiration, const ISendMessagesRef& Sender) override;
 	virtual void Shutdown() override;
-	virtual IMessageSubscriptionPtr Subscribe( const IReceiveMessagesRef& Subscriber, const FName& MessageType, const FMessageScopeRange& ScopeRange ) override;
-	virtual void Unintercept( const IMessageInterceptorRef& Interceptor, const FName& MessageType ) override;
-	virtual void Unregister( const FMessageAddress& Address ) override;
-	virtual void Unsubscribe( const IReceiveMessagesRef& Subscriber, const FName& MessageType ) override;
+	virtual IMessageSubscriptionPtr Subscribe(const IReceiveMessagesRef& Subscriber, const FName& MessageType, const FMessageScopeRange& ScopeRange) override;
+	virtual void Unintercept(const IMessageInterceptorRef& Interceptor, const FName& MessageType) override;
+	virtual void Unregister(const FMessageAddress& Address) override;
+	virtual void Unsubscribe(const IReceiveMessagesRef& Subscriber, const FName& MessageType) override;
 
 private:
 
diff --git a/Engine/Source/Runtime/Messaging/Private/Bus/MessageDispatchTask.h b/Engine/Source/Runtime/Messaging/Private/Bus/MessageDispatchTask.h
--- a/Engine/Source/Runtime/Messaging/Private/Bus/MessageDispatchTask.h
+++ b/Engine/Source/Runtime/Messaging/Private/Bus/MessageDispatchTask.h
@@ -18,7 +18,7 @@ public:
 	 * @param InRecipient The message recipient.
 	 * @param InTracer The message tracer to notify.
 	 */
-	FMessageDispatchTask( ENamedThreads::Type InThread, IMessageContextRef InContext, IReceiveMessagesWeakPtr InRecipient, FMessageTracerPtr InTracer )
+	FMessageDispatchTask(ENamedThreads::Type InThread, IMessageContextRef InContext, IReceiveMessagesWeakPtr InRecipient, FMessageTracerPtr InTracer)
 		: Context(InContext)
 		, RecipientPtr(InRecipient)
 		, Thread(InThread)
@@ -33,25 +33,27 @@ public:
 	 * @param CurrentThread The thread that this task is executing on.
 	 * @param MyCompletionGraphEvent The completion event.
 	 */
-	void DoTask( ENamedThreads::Type CurrentThread, const FGraphEventRef& MyCompletionGraphEvent )
+	void DoTask(ENamedThreads::Type CurrentThread, const FGraphEventRef& MyCompletionGraphEvent)
 	{
 		IReceiveMessagesPtr Recipient = RecipientPtr.Pin();
 
-		if (Recipient.IsValid())
+		if (!Recipient.IsValid())
 		{
-			FMessageTracerPtr Tracer = TracerPtr.Pin();
+			return;
+		}
+
+		FMessageTracerPtr Tracer = TracerPtr.Pin();
 
-			if (Tracer.IsValid())
-			{
-				Tracer->TraceDispatchedMessage(Context, Recipient.ToSharedRef(), true);
-			}
+		if (Tracer.IsValid())
+		{
+			Tracer->TraceDispatchedMessage(Context, Recipient.ToSharedRef(), true);
+		}
 		
-			Recipient->ReceiveMessage(Context);
+		Recipient->ReceiveMessage(Context);
 
-			if (TracerPtr.IsValid())
-			{
-				Tracer->TraceHandledMessage(Context, Recipient.ToSharedRef());
-			}
+		if (TracerPtr.IsValid())
+		{
+			Tracer->TraceHandledMessage(Context, Recipient.ToSharedRef());
 		}
 	}
 	
diff --git a/Engine/Source/Runtime/Messaging/Private/Bus/MessageRouter.cpp b/Engine/Source/Runtime/Messaging/Private/Bus/MessageRouter.cpp
--- a/Engine/Source/Runtime/Messaging/Private/Bus/MessageRouter.cpp
+++ b/Engine/Source/Runtime/Messaging/Private/Bus/MessageRouter.cpp
@@ -72,18 +72,16 @@ void FMessageRouter::Exit()
 	TArray<IReceiveMessagesWeakPtr> Recipients;
 
 	// gather all subscribed and registered recipients
-	for (TMap<FMessageAddress, IReceiveMessagesWeakPtr>::TConstIterator It(ActiveRecipients); It; ++It)
+	for (const auto& RecipientPair : ActiveRecipients)
 	{
-		Recipients.AddUnique(It.Value());
+		Recipients.AddUnique(RecipientPair.Value);
 	}
 
-	for (TMap<FName, TArray<IMessageSubscriptionPtr> >::TIterator It(ActiveSubscriptions); It; ++It)
+	for (const auto& SubscriptionsPair : ActiveSubscriptions)
 	{
-		TArray<IMessageSubscriptionPtr>& Subscribers = It.Value();
-
-		for (int32 Index = 0; Index < Subscribers.Num(); ++Index)
+		for (const auto& Subscription : SubscriptionsPair.Value)
 		{
-			Recipients.AddUnique(Subscribers[Index]->GetSubscriber());
+			Recipients.AddUnique(Subscription->GetSubscriber());
 		}
 	}
 }
@@ -121,9 +119,9 @@ void FMessageRouter::DispatchMessage( const IMessageContextRef& Context )
 
 		if (RecipientList.Num() > 0)
 		{
-			for (int32 Index = 0; Index < RecipientList.Num(); Index++)
+			for (const auto& RecipientAddress : RecipientList)
 			{
-				IReceiveMessagesPtr Recipient = ActiveRecipients.FindRef(RecipientList[Index]).Pin();
+				IReceiveMessagesPtr Recipient = ActiveRecipients.FindRef(RecipientAddress).Pin();
 
 				if (Recipient.IsValid())
 				{
@@ -131,7 +129,7 @@ void FMessageRouter::DispatchMessage( const IMessageContextRef& Context )
 				}
 				else
 				{
-					ActiveRecipients.Remove(RecipientList[Index]);
+					ActiveRecipients.Remove(RecipientAddress);
 				}
 			}
 		}
@@ -143,9 +141,8 @@ void FMessageRouter::DispatchMessage( const IMessageContextRef& Context )
 		}
 
 		// dispatch the message
-		for (int32 RecipientIndex = 0; RecipientIndex < Recipients.Num(); RecipientIndex++)
+		for (auto& Recipient : Recipients)
 		{
-			IReceiveMessagesPtr Recipient = Recipients[RecipientIndex];
 			ENamedThreads::Type RecipientThread = Recipient->GetRecipientThread();
 
 			if (RecipientThread == ENamedThreads::AnyThread)
@@ -171,30 +168,32 @@ void FMessageRouter::FilterSubscriptions( TArray<IMessageSubscriptionPtr>& Subsc
 	{
 		const IMessageSubscriptionPtr& Subscription = Subscriptions[SubscriptionIndex];
 
-		if (Subscription->IsEnabled() && Subscription->GetScopeRange().Contains(MessageScope))
+		if (!Subscription->IsEnabled() || !Subscription->GetScopeRange().Contains(MessageScope))
 		{
-			IReceiveMessagesPtr Subscriber = Subscription->GetSubscriber().Pin();
+			continue;
+		}
+
+		IReceiveMessagesPtr Subscriber = Subscription->GetSubscriber().Pin();
 
-			if (Subscriber.IsValid())
+		if (Subscriber.IsValid())
+		{
+			if (MessageScope == EMessageScope::Thread)
 			{
-				if (MessageScope == EMessageScope::Thread)
-				{
-					ENamedThreads::Type RecipientThread = Subscriber->GetRecipientThread();
-					ENamedThreads::Type SenderThread = Context->GetSenderThread();
+				ENamedThreads::Type RecipientThread = Subscriber->GetRecipientThread();
+				ENamedThreads::Type SenderThread = Context->GetSenderThread();
 
-					if (RecipientThread != SenderThread)
-					{
-						continue;
-					}
+				if (RecipientThread != SenderThread)
+				{
+					continue;
 				}
-
-				OutRecipients.AddUnique(Subscriber);
-			}
-			else
-			{
-				Subscriptions.RemoveAtSwap(SubscriptionIndex);
-				--SubscriptionIndex;
 			}
+
+			OutRecipients.AddUnique(Subscriber);
+		}
+		else
+		{
+			Subscriptions.RemoveAtSwap(SubscriptionIndex);
+			--SubscriptionIndex;
 		}
 	}
 }
@@ -245,10 +244,9 @@ void FMessageRouter::HandleRemoveInterceptor( IMessageInterceptorRef Interceptor
 {
 	if (MessageType == NAME_All)
 	{
-		for (TMap<FName, TArray<IMessageInterceptorPtr> >::TIterator It(ActiveInterceptors); It; ++It)
+		for (auto& InterceptorsPair : ActiveInterceptors)
 		{
-			TArray<IMessageInterceptorPtr>& Interceptors = It.Value();
-			Interceptors.Remove(Interceptor);
+			InterceptorsPair.Value.Remove(Interceptor);
 		}
 	}
 	else
@@ -278,25 +276,30 @@ void FMessageRouter::HandleRemoveSubscriber( IReceiveMessagesWeakPtr SubscriberP
 {
 	IReceiveMessagesPtr Subscriber = SubscriberPtr.Pin();
 
-	if (Subscriber.IsValid())
+	if (!Subscriber.IsValid())
+	{
+		return;
+	}
+
+	for (auto& SubscriptionsPair : ActiveSubscriptions)
 	{
-		for (TMap<FName, TArray<IMessageSubscriptionPtr> >::TIterator It(ActiveSubscriptions); It; ++It)
+		if ((MessageType != NAME_All) && (MessageType != SubscriptionsPair.Key))
 		{
-			if ((MessageType == NAME_All) || (MessageType == It.Key()))
-			{
-				TArray<IMessageSubscriptionPtr>& Subsriptions = It.Value();
+			continue;
+		}
 
-				for (int32 Index = 0; Index < Subsriptions.Num(); Index++)
-				{
-					if (Subsriptions[Index]->GetSubscriber().Pin() == Subscriber)
-					{
-						Subsriptions.RemoveAtSwap(Index);
+		TArray<IMessageSubscriptionPtr>& Subscriptions = SubscriptionsPair.Value;
+
+		for (int32 Index = 0; Index < Subscriptions.Num(); Index++)
+		{
+			const auto& Subscription = Subscriptions[Index];
 
-						Tracer->TraceRemovedSubscription(Subsriptions[Index].ToSharedRef(), MessageType);
+			if (Subscription->GetSubscriber().Pin() == Subscriber)
+			{
+				Subscriptions.RemoveAtSwap(Index);
+				Tracer->TraceRemovedSubscription(Subscription.ToSharedRef(), MessageType);
 
-						break;
-					}
-				}
+				break;
 			}
 		}
 	}
@@ -308,11 +311,11 @@ void FMessageRouter::HandleRouteMessage( IMessageContextRef Context )
 	// intercept routing
 	TArray<IMessageInterceptorPtr>& Interceptors = ActiveInterceptors.FindOrAdd(Context->GetMessageType());
 
-	for (TArray<IMessageInterceptorPtr>::TIterator It(Interceptors); It; ++It)
+	for (auto& Interceptor : Interceptors)
 	{
-		if ((*It)->InterceptMessage(Context))
+		if (Interceptor->InterceptMessage(Context))
 		{
-			Tracer->TraceInterceptedMessage(Context, It->ToSharedRef());
+			Tracer->TraceInterceptedMessage(Context, Interceptor.ToSharedRef());
 
 			return;
 		}
diff --git a/Engine/Source/Runtime/Messaging/Private/Bus/MessageTracer.cpp b/Engine/Source/Runtime/Messaging/Private/Bus/MessageTracer.cpp
--- a/Engine/Source/Runtime/Messaging/Private/Bus/MessageTracer.cpp
+++ b/Engine/Source/Runtime/Messaging/Private/Bus/MessageTracer.cpp
@@ -27,7 +27,7 @@ FMessageTracer::~FMessageTracer()
 /* FMessageTracer interface
  *****************************************************************************/
 
-void FMessageTracer::TraceAddedInterceptor( const IMessageInterceptorRef& Interceptor, const FName& MessageType )
+void FMessageTracer::TraceAddedInterceptor(const IMessageInterceptorRef& Interceptor, const FName& MessageType)
 {
 	double Timestamp = FPlatformTime::Seconds();
 
@@ -48,7 +48,7 @@ void FMessageTracer::TraceAddedInterceptor( const IMessageInterceptorRef& Interc
 }
 
 
-void FMessageTracer::TraceAddedRecipient( const FMessageAddress& Address, const IReceiveMessagesRef& Recipient )
+void FMessageTracer::TraceAddedRecipient(const FMessageAddress& Address, const IReceiveMessagesRef& Recipient)
 {
 	double Timestamp = FPlatformTime::Seconds();
 
@@ -63,9 +63,11 @@ void FMessageTracer::TraceAddedRecipient( const FMessageAddress& Address, const
 
 		// initialize endpoint information
 		FMessageTracerAddressInfoRef AddressInfo = MakeShareable(new FMessageTracerAddressInfo());
-		AddressInfo->Address = Address;
-		AddressInfo->TimeRegistered = Timestamp;
-		AddressInfo->TimeUnregistered = 0;
+		{
+			AddressInfo->Address = Address;
+			AddressInfo->TimeRegistered = Timestamp;
+			AddressInfo->TimeUnregistered = 0;
+		}
 
 		EndpointInfo->AddressInfos.Add(Address, AddressInfo);
 		EndpointInfo->Name = Recipient->GetDebugName();
@@ -77,7 +79,7 @@ void FMessageTracer::TraceAddedRecipient( const FMessageAddress& Address, const
 }
 
 
-void FMessageTracer::TraceAddedSubscription( const IMessageSubscriptionRef& Subscription )
+void FMessageTracer::TraceAddedSubscription(const IMessageSubscriptionRef& Subscription)
 {
 	if (!Running)
 	{
@@ -92,7 +94,7 @@ void FMessageTracer::TraceAddedSubscription( const IMessageSubscriptionRef& Subs
 }
 
 
-void FMessageTracer::TraceDispatchedMessage( const IMessageContextRef& Context, const IReceiveMessagesRef& Recipient, bool Async )
+void FMessageTracer::TraceDispatchedMessage(const IMessageContextRef& Context, const IReceiveMessagesRef& Recipient, bool Async)
 {
 	if (!Running)
 	{
@@ -119,13 +121,14 @@ void FMessageTracer::TraceDispatchedMessage( const IMessageContextRef& Context,
 
 		// update message information
 		FMessageTracerDispatchStateRef DispatchState = MakeShareable(new FMessageTracerDispatchState());
-
-		DispatchState->DispatchLatency = Timestamp - MessageInfo->TimeSent;
-		DispatchState->DispatchType = Async ? EMessageTracerDispatchTypes::TaskGraph : EMessageTracerDispatchTypes::Direct;
-		DispatchState->EndpointInfo = EndpointInfo;
-		DispatchState->RecipientThread = Recipient->GetRecipientThread();
-		DispatchState->TimeDispatched = Timestamp;
-		DispatchState->TimeHandled = 0.0;
+		{
+			DispatchState->DispatchLatency = Timestamp - MessageInfo->TimeSent;
+			DispatchState->DispatchType = Async ? EMessageTracerDispatchTypes::TaskGraph : EMessageTracerDispatchTypes::Direct;
+			DispatchState->EndpointInfo = EndpointInfo;
+			DispatchState->RecipientThread = Recipient->GetRecipientThread();
+			DispatchState->TimeDispatched = Timestamp;
+			DispatchState->TimeHandled = 0.0;
+		}
 
 		MessageInfo->DispatchStates.Add(EndpointInfo, DispatchState);
 
@@ -135,7 +138,7 @@ void FMessageTracer::TraceDispatchedMessage( const IMessageContextRef& Context,
 }
 
 
-void FMessageTracer::TraceHandledMessage( const IMessageContextRef& Context, const IReceiveMessagesRef& Recipient )
+void FMessageTracer::TraceHandledMessage(const IMessageContextRef& Context, const IReceiveMessagesRef& Recipient)
 {
 	if (!Running)
 	{
@@ -171,7 +174,7 @@ void FMessageTracer::TraceHandledMessage( const IMessageContextRef& Context, con
 }
 
 
-void FMessageTracer::TraceInterceptedMessage( const IMessageContextRef& Context, const IMessageInterceptorRef& Interceptor )
+void FMessageTracer::TraceInterceptedMessage(const IMessageContextRef& Context, const IMessageInterceptorRef& Interceptor)
 {
 	if (!Running)
 	{
@@ -204,7 +207,7 @@ void FMessageTracer::TraceInterceptedMessage( const IMessageContextRef& Context,
 }
 
 
-void FMessageTracer::TraceRemovedInterceptor( const IMessageInterceptorRef& Interceptor, const FName& MessageType )
+void FMessageTracer::TraceRemovedInterceptor(const IMessageInterceptorRef& Interceptor, const FName& MessageType)
 {
 	double Timestamp = FPlatformTime::Seconds();
 
@@ -222,7 +225,7 @@ void FMessageTracer::TraceRemovedInterceptor( const IMessageInterceptorRef& Inte
 }
 
 
-void FMessageTracer::TraceRemovedRecipient( const FMessageAddress& Address )
+void FMessageTracer::TraceRemovedRecipient(const FMessageAddress& Address)
 {
 	double Timestamp = FPlatformTime::Seconds();
 
@@ -245,7 +248,7 @@ void FMessageTracer::TraceRemovedRecipient( const FMessageAddress& Address )
 }
 
 
-void FMessageTracer::TraceRemovedSubscription( const IMessageSubscriptionRef& Subscription, const FName& MessageType )
+void FMessageTracer::TraceRemovedSubscription(const IMessageSubscriptionRef& Subscription, const FName& MessageType)
 {
 	if (!Running)
 	{
@@ -260,7 +263,7 @@ void FMessageTracer::TraceRemovedSubscription( const IMessageSubscriptionRef& Su
 }
 
 
-void FMessageTracer::TraceRoutedMessage( const IMessageContextRef& Context )
+void FMessageTracer::TraceRoutedMessage(const IMessageContextRef& Context)
 {
 	if (!Running)
 	{
@@ -287,7 +290,7 @@ void FMessageTracer::TraceRoutedMessage( const IMessageContextRef& Context )
 }
 
 
-void FMessageTracer::TraceSentMessage( const IMessageContextRef& Context )
+void FMessageTracer::TraceSentMessage(const IMessageContextRef& Context)
 {
 	if (!Running)
 	{
@@ -307,13 +310,14 @@ void FMessageTracer::TraceSentMessage( const IMessageContextRef& Context )
 
 		// create message info
 		FMessageTracerMessageInfoRef MessageInfo = MakeShareable(new FMessageTracerMessageInfo());
-	
-		MessageInfo->Context = Context;
-		MessageInfo->Intercepted = false;
-		MessageInfo->SenderInfo = EndpointInfo;
-		MessageInfo->TimeRouted = 0.0;
-		MessageInfo->TimeSent = Timestamp;
-		MessageInfos.Add(Context, MessageInfo);
+		{
+			MessageInfo->Context = Context;
+			MessageInfo->Intercepted = false;
+			MessageInfo->SenderInfo = EndpointInfo;
+			MessageInfo->TimeRouted = 0.0;
+			MessageInfo->TimeSent = Timestamp;
+			MessageInfos.Add(Context, MessageInfo);
+		}
 
 		// add message type
 		FMessageTracerTypeInfoPtr& TypeInfo = MessageTypes.FindOrAdd(Context->GetMessageType());
@@ -340,7 +344,7 @@ void FMessageTracer::TraceSentMessage( const IMessageContextRef& Context )
 /* IMessageTracer interface
  *****************************************************************************/
 
-int32 FMessageTracer::GetEndpoints( TArray<FMessageTracerEndpointInfoPtr>& OutEndpoints ) const
+int32 FMessageTracer::GetEndpoints(TArray<FMessageTracerEndpointInfoPtr>& OutEndpoints) const
 {
 	RecipientsToEndpointInfos.GenerateValueArray(OutEndpoints);
 
@@ -348,7 +352,7 @@ int32 FMessageTracer::GetEndpoints( TArray<FMessageTracerEndpointInfoPtr>& OutEn
 }
 
 
-int32 FMessageTracer::GetMessages( TArray<FMessageTracerMessageInfoPtr>& OutMessages ) const
+int32 FMessageTracer::GetMessages(TArray<FMessageTracerMessageInfoPtr>& OutMessages) const
 {
 	MessageInfos.GenerateValueArray(OutMessages);
 
@@ -356,7 +360,7 @@ int32 FMessageTracer::GetMessages( TArray<FMessageTracerMessageInfoPtr>& OutMess
 }
 
 
-int32 FMessageTracer::GetMessageTypes( TArray<FMessageTracerTypeInfoPtr>& OutTypes ) const
+int32 FMessageTracer::GetMessageTypes(TArray<FMessageTracerTypeInfoPtr>& OutTypes) const
 {
 	MessageTypes.GenerateValueArray(OutTypes);
 
@@ -364,7 +368,7 @@ int32 FMessageTracer::GetMessageTypes( TArray<FMessageTracerTypeInfoPtr>& OutTyp
 }
 
 
-bool FMessageTracer::Tick( float DeltaTime )
+bool FMessageTracer::Tick(float DeltaTime)
 {
 	if (ResetPending)
 	{
@@ -395,28 +399,28 @@ void FMessageTracer::ResetMessages()
 	MessageInfos.Reset();
 	MessageTypes.Reset();
 
-	for (TMap<FMessageAddress, FMessageTracerEndpointInfoPtr>::TIterator It(AddressesToEndpointInfos); It; ++It)
+	for (auto& EndpointInfoPair : AddressesToEndpointInfos)
 	{
-		FMessageTracerEndpointInfoPtr& EndpointInfo = It.Value();
-		EndpointInfo->ReceivedMessages.Reset();
-		EndpointInfo->SentMessages.Reset();
+		FMessageTracerEndpointInfoPtr& EndpointInfo = EndpointInfoPair.Value;
+		{
+			EndpointInfo->ReceivedMessages.Reset();
+			EndpointInfo->SentMessages.Reset();
+		}
 	}
 
 	MessagesResetDelegate.Broadcast();
 }
 
 
-bool FMessageTracer::ShouldBreak( const IMessageContextRef& Context ) const
+bool FMessageTracer::ShouldBreak(const IMessageContextRef& Context) const
 {
 	if (Breaking)
 	{
 		return true;
 	}
 
-	for (int32 BreakpointIndex = Breakpoints.Num() - 1; BreakpointIndex >= 0; --BreakpointIndex)
+	for (const auto& Breakpoint : Breakpoints)
 	{
-		const IMessageTracerBreakpointPtr& Breakpoint = Breakpoints[BreakpointIndex];
-
 		if (Breakpoint->IsEnabled() && Breakpoint->ShouldBreak(Context))
 		{
 			return true;
