Messaging: Fixed message bridge not unregistering remote endpoints on shutdown

diff --git a/Engine/Source/Runtime/Messaging/Private/Bridge/MessageBridge.cpp b/Engine/Source/Runtime/Messaging/Private/Bridge/MessageBridge.cpp
--- a/Engine/Source/Runtime/Messaging/Private/Bridge/MessageBridge.cpp
+++ b/Engine/Source/Runtime/Messaging/Private/Bridge/MessageBridge.cpp
@@ -29,10 +29,11 @@ FMessageBridge::~FMessageBridge()
 		Bus->Unregister(Address);
 
 		TArray<FMessageAddress> RemovedAddresses;
+		AddressBook.RemoveAll(RemovedAddresses);
 
-		for (int32 AddressIndex = 0; AddressIndex < RemovedAddresses.Num(); ++AddressIndex)
+		for (const auto& Address : RemovedAddresses)
 		{
-			Bus->Unregister(RemovedAddresses[AddressIndex]);
+			Bus->Unregister(Address);
 		}
 	}
 }
