MessagingDebugger: Continued to implement breakpoints.

diff --git a/Engine/Plugins/Messaging/MessagingDebugger/Source/MessagingDebugger/Private/Styles/MessagingDebuggerStyle.h b/Engine/Plugins/Messaging/MessagingDebugger/Source/MessagingDebugger/Private/Styles/MessagingDebuggerStyle.h
--- a/Engine/Plugins/Messaging/MessagingDebugger/Source/MessagingDebugger/Private/Styles/MessagingDebuggerStyle.h
+++ b/Engine/Plugins/Messaging/MessagingDebugger/Source/MessagingDebugger/Private/Styles/MessagingDebuggerStyle.h
@@ -61,6 +61,7 @@ public:
 
 		// misc icons
 		Set("Break", new IMAGE_BRUSH( "icon_break_16x", Icon16x16));
+		Set("BreakColumn", new IMAGE_BRUSH( "BreakpointBorder", Icon16x16));
 		Set("BreakDisabled", new IMAGE_BRUSH( "icon_break_disabled_16x", Icon16x16));
 		Set("BreakIn", new IMAGE_BRUSH( "icon_break_in_16x", Icon16x16));
 		Set("BreakOut", new IMAGE_BRUSH( "icon_break_out_16x", Icon16x16));
diff --git a/Engine/Plugins/Messaging/MessagingDebugger/Source/MessagingDebugger/Private/Widgets/Endpoints/SMessagingEndpoints.cpp b/Engine/Plugins/Messaging/MessagingDebugger/Source/MessagingDebugger/Private/Widgets/Endpoints/SMessagingEndpoints.cpp
--- a/Engine/Plugins/Messaging/MessagingDebugger/Source/MessagingDebugger/Private/Widgets/Endpoints/SMessagingEndpoints.cpp
+++ b/Engine/Plugins/Messaging/MessagingDebugger/Source/MessagingDebugger/Private/Widgets/Endpoints/SMessagingEndpoints.cpp
@@ -74,7 +74,7 @@ void SMessagingEndpoints::Construct( const FArguments& InArgs, const FMessagingD
 									.HeaderContent()
 									[
 										SNew(SImage)
-											.Image(InStyle->GetBrush("Break"))
+											.Image(InStyle->GetBrush("BreakColumn"))
 											.ToolTipText(LOCTEXT("EndpointListBreakpointsColumnTooltip", "Breakpoints"))
 									]
 
@@ -113,6 +113,8 @@ void SMessagingEndpoints::Construct( const FArguments& InArgs, const FMessagingD
 	Filter->OnChanged().AddRaw(this, &SMessagingEndpoints::HandleFilterChanged);
 	Model->OnSelectedMessageChanged().AddRaw(this, &SMessagingEndpoints::HandleModelSelectedMessageChanged);
 }
+
+
 void SMessagingEndpoints::Tick( const FGeometry& AllottedGeometry, const double InCurrentTime, const float InDeltaTime )
 {
 	// @todo gmp: fix this
diff --git a/Engine/Plugins/Messaging/MessagingDebugger/Source/MessagingDebugger/Private/Widgets/Types/SMessagingTypes.cpp b/Engine/Plugins/Messaging/MessagingDebugger/Source/MessagingDebugger/Private/Widgets/Types/SMessagingTypes.cpp
--- a/Engine/Plugins/Messaging/MessagingDebugger/Source/MessagingDebugger/Private/Widgets/Types/SMessagingTypes.cpp
+++ b/Engine/Plugins/Messaging/MessagingDebugger/Source/MessagingDebugger/Private/Widgets/Types/SMessagingTypes.cpp
@@ -79,8 +79,8 @@ void SMessagingTypes::Construct( const FArguments& InArgs, const FMessagingDebug
 									.HeaderContent()
 									[
 										SNew(SImage)
-											.Image(Style->GetBrush("Break"))
-											.ToolTipText(LOCTEXT("TypeListBreakpointsColumnTooltip", "Breakpoints"))
+											.Image(InStyle->GetBrush("BreakColumn"))
+											.ToolTipText(LOCTEXT("EndpointListBreakpointsColumnTooltip", "Breakpoints"))
 									]
 
 								+ SHeaderRow::Column("Name")
diff --git a/Engine/Source/Runtime/Messaging/Private/Bus/MessageTracer.h b/Engine/Source/Runtime/Messaging/Private/Bus/MessageTracer.h
--- a/Engine/Source/Runtime/Messaging/Private/Bus/MessageTracer.h
+++ b/Engine/Source/Runtime/Messaging/Private/Bus/MessageTracer.h
@@ -205,9 +205,7 @@ public:
 
 protected:
 
-	/**
-	 * Resets traced messages.
-	 */
+	/** Resets traced messages. */
 	void ResetMessages();
 
 	/**
@@ -228,9 +226,6 @@ private:
 	/** Holds the collection of breakpoints. */
 	TArray<IMessageTracerBreakpointPtr> Breakpoints;
 
-	/** Holds the collection of message senders to break on when they send a message. */
-	TArray<FMessageAddress> BreakOnSenders;
-
 	/** Holds an event signaling that messaging routing can continue. */
 	FEvent* ContinueEvent;
 
diff --git a/Engine/Source/Runtime/Messaging/Public/IMessageTracer.h b/Engine/Source/Runtime/Messaging/Public/IMessageTracer.h
--- a/Engine/Source/Runtime/Messaging/Public/IMessageTracer.h
+++ b/Engine/Source/Runtime/Messaging/Public/IMessageTracer.h
@@ -43,6 +43,23 @@ typedef TSharedPtr<struct FMessageTracerTypeInfo> FMessageTracerTypeInfoPtr;
 typedef TSharedRef<struct FMessageTracerTypeInfo> FMessageTracerTypeInfoRef;
 
 
+/** Enumerates tracer breakpoint states. */
+enum class EMessageTracerBreakpointState
+{
+	/** The breakpoint is disabled. */
+	Disabled,
+
+	/** The breakpoint is enabled. */
+	Enabled,
+
+	/** The breakpoint is enabled for incoming messages. */
+	EnabledIn,
+
+	/** The breakpoint is enabled for outgoing messages. */
+	EnabledOut,
+};
+
+
 /**
  * Enumerates message dispatch types.
  */
@@ -59,6 +76,25 @@ enum class EMessageTracerDispatchTypes
 };
 
 
+/**
+ * Structure for tracer breakpoints.
+ */
+struct FMessageTracerBreakpoint
+{
+	/** Recipient address to break on. */
+	FMessageAddress BreakOnRecipient;
+
+	/** Sender address to break on. */
+	FMessageAddress BreakOnSender;
+
+	/** How many times the breakpoint was hit. */
+	int64 HitCount;
+
+	/** The breakpoint's enabled state. */
+	EMessageTracerBreakpointState State;
+};
+
+
 /**
  * Structure for message dispatch states.
  */
