Fixed shadowed variables

diff --git a/Engine/Source/Runtime/Messaging/Private/Bridge/MessageBridge.cpp b/Engine/Source/Runtime/Messaging/Private/Bridge/MessageBridge.cpp
--- a/Engine/Source/Runtime/Messaging/Private/Bridge/MessageBridge.cpp
+++ b/Engine/Source/Runtime/Messaging/Private/Bridge/MessageBridge.cpp
@@ -31,9 +31,9 @@ FMessageBridge::~FMessageBridge()
 		TArray<FMessageAddress> RemovedAddresses;
 		AddressBook.RemoveAll(RemovedAddresses);
 
-		for (const auto& Address : RemovedAddresses)
+		for (const auto& RemovedAddress : RemovedAddresses)
 		{
-			Bus->Unregister(Address);
+			Bus->Unregister(RemovedAddress);
 		}
 	}
 }
@@ -189,9 +189,9 @@ void FMessageBridge::HandleTransportNodeLost(const FGuid& LostNodeId)
 	// unregister endpoints
 	if (Bus.IsValid())
 	{
-		for (const auto& Address : RemovedAddresses)
+		for (const auto& RemovedAddress : RemovedAddresses)
 		{
-			Bus->Unregister(Address);
+			Bus->Unregister(RemovedAddress);
 		}
 	}
 }
