Messaging: Added documentation for event delegate parameters

diff --git a/Engine/Source/Runtime/Messaging/Public/IMessageTransport.h b/Engine/Source/Runtime/Messaging/Public/IMessageTransport.h
--- a/Engine/Source/Runtime/Messaging/Public/IMessageTransport.h
+++ b/Engine/Source/Runtime/Messaging/Public/IMessageTransport.h
@@ -46,32 +46,20 @@ public:
 	 * @param Recipients The transport nodes to send the message to.
 	 * @return true if the message is being transported, false otherwise.
 	 */
-	virtual bool TransportMessage( const IMessageContextRef& Context, const TArray<FGuid>& Recipients ) = 0;
+	virtual bool TransportMessage(const IMessageContextRef& Context, const TArray<FGuid>& Recipients) = 0;
 
 public:
 
-	/**
-	 * A delegate that is executed when message data has been received.
-	 *
-	 * @param Delegate The delegate to set.
-	 */
-	DECLARE_DELEGATE_TwoParams(FOnMessageReceived, const IMessageContextRef&, const FGuid&)
+	/** A delegate that is executed when message data has been received. */
+	DECLARE_DELEGATE_TwoParams(FOnMessageReceived, const IMessageContextRef& /*MessageContext*/, const FGuid& /*NodeId*/)
 	virtual FOnMessageReceived& OnMessageReceived() = 0;
 
-	/**
-	 * A delegate that is executed when a transport node has been discovered.
-	 *
-	 * @param Delegate The delegate to set.
-	 */
-	DECLARE_DELEGATE_OneParam(FOnNodeDiscovered, const FGuid&)
+	/** A delegate that is executed when a transport node has been discovered. */
+	DECLARE_DELEGATE_OneParam(FOnNodeDiscovered, const FGuid& /*NodeId*/)
 	virtual FOnNodeDiscovered& OnNodeDiscovered() = 0;
 
-	/**
-	 * A delegate that is executed when a transport node has closed or timed out.
-	 *
-	 * @param Delegate The delegate to set.
-	 */
-	DECLARE_DELEGATE_OneParam(FOnNodeLost, const FGuid&)
+	/** A delegate that is executed when a transport node has closed or timed out. */
+	DECLARE_DELEGATE_OneParam(FOnNodeLost, const FGuid& /*NodeId*/)
 	virtual FOnNodeLost& OnNodeLost() = 0;
 
 protected:
