Copying //UE4/Dev-Sequencer to //UE4/Dev-Main (Source: //UE4/Dev-Sequencer @ 3178529)

diff --git a/Engine/Source/Runtime/Messaging/Private/Bus/MessageRouter.h b/Engine/Source/Runtime/Messaging/Private/Bus/MessageRouter.h
--- a/Engine/Source/Runtime/Messaging/Private/Bus/MessageRouter.h
+++ b/Engine/Source/Runtime/Messaging/Private/Bus/MessageRouter.h
@@ -3,11 +3,12 @@
 #pragma once
 
 #include "IMessageContext.h"
-#include "IMessageInterceptor.h"
-#include "IMessageSubscription.h"
-#include "IMessageTracer.h"
-#include "IReceiveMessages.h"
-#include "ISendMessages.h"
+#include "MessageTracer.h"
+
+
+class IMessageInterceptor;
+class IMessageSubscription;
+class IMessageTracer;
 
 
 /**
@@ -34,7 +35,7 @@ public:
 	 * @param Interceptor The interceptor to add.
 	 * @param MessageType The type of messages to intercept.
 	 */
-	FORCEINLINE void AddInterceptor(const IMessageInterceptorRef& Interceptor, const FName& MessageType)
+	FORCEINLINE void AddInterceptor(const TSharedRef<IMessageInterceptor, ESPMode::ThreadSafe>& Interceptor, const FName& MessageType)
 	{
 		EnqueueCommand(FSimpleDelegate::CreateRaw(this, &FMessageRouter::HandleAddInterceptor, Interceptor, MessageType));
 	}
@@ -45,9 +46,9 @@ public:
 	 * @param Address The address of the recipient to add.
 	 * @param Recipient The recipient.
 	 */
-	FORCEINLINE void AddRecipient(const FMessageAddress& Address, const IReceiveMessagesRef& Recipient)
+	FORCEINLINE void AddRecipient(const FMessageAddress& Address, const TSharedRef<IMessageReceiver, ESPMode::ThreadSafe>& Recipient)
 	{
-		EnqueueCommand(FSimpleDelegate::CreateRaw(this, &FMessageRouter::HandleAddRecipient, Address, IReceiveMessagesWeakPtr(Recipient)));
+		EnqueueCommand(FSimpleDelegate::CreateRaw(this, &FMessageRouter::HandleAddRecipient, Address, TWeakPtr<IMessageReceiver, ESPMode::ThreadSafe>(Recipient)));
 	}
 
 	/**
@@ -55,7 +56,7 @@ public:
 	 *
 	 * @param Subscription The subscription to add.
 	 */
-	FORCEINLINE void AddSubscription(const IMessageSubscriptionRef& Subscription)
+	FORCEINLINE void AddSubscription(const TSharedRef<IMessageSubscription, ESPMode::ThreadSafe>& Subscription)
 	{
 		EnqueueCommand(FSimpleDelegate::CreateRaw(this, &FMessageRouter::HandleAddSubscriber, Subscription));
 	}
@@ -65,7 +66,7 @@ public:
 	 *
 	 * @return Weak pointer to the message tracer.
 	 */
-	FORCEINLINE IMessageTracerRef GetTracer()
+	FORCEINLINE TSharedRef<IMessageTracer, ESPMode::ThreadSafe> GetTracer()
 	{
 		return Tracer;
 	}
@@ -76,7 +77,7 @@ public:
 	 * @param Interceptor The interceptor to remove.
 	 * @param MessageType The type of messages to stop intercepting.
 	 */
-	FORCEINLINE void RemoveInterceptor(const IMessageInterceptorRef& Interceptor, const FName& MessageType)
+	FORCEINLINE void RemoveInterceptor(const TSharedRef<IMessageInterceptor, ESPMode::ThreadSafe>& Interceptor, const FName& MessageType)
 	{
 		EnqueueCommand(FSimpleDelegate::CreateRaw(this, &FMessageRouter::HandleRemoveInterceptor, Interceptor, MessageType));
 	}
@@ -97,9 +98,9 @@ public:
 	 * @param Subscriber The subscriber to stop routing messages to.
 	 * @param MessageType The type of message to unsubscribe from (NAME_None = all types).
 	 */
-	FORCEINLINE void RemoveSubscription(const IReceiveMessagesRef& Subscriber, const FName& MessageType)
+	FORCEINLINE void RemoveSubscription(const TSharedRef<IMessageReceiver, ESPMode::ThreadSafe>& Subscriber, const FName& MessageType)
 	{
-		EnqueueCommand(FSimpleDelegate::CreateRaw(this, &FMessageRouter::HandleRemoveSubscriber, IReceiveMessagesWeakPtr(Subscriber), MessageType));
+		EnqueueCommand(FSimpleDelegate::CreateRaw(this, &FMessageRouter::HandleRemoveSubscriber, TWeakPtr<IMessageReceiver, ESPMode::ThreadSafe>(Subscriber), MessageType));
 	}
 
 	/**
@@ -107,7 +108,7 @@ public:
 	 *
 	 * @param Context The context of the message to route.
 	 */
-	FORCEINLINE void RouteMessage(const IMessageContextRef& Context)
+	FORCEINLINE void RouteMessage(const TSharedRef<IMessageContext, ESPMode::ThreadSafe>& Context)
 	{
 		Tracer->TraceSentMessage(Context);
 		EnqueueCommand(FSimpleDelegate::CreateRaw(this, &FMessageRouter::HandleRouteMessage, Context));
@@ -115,7 +116,7 @@ public:
 
 public:
 
-	// FRunnable interface
+	//~ FRunnable interface
 
 	virtual bool Init() override;
 	virtual uint32 Run() override;
@@ -157,40 +158,43 @@ protected:
 	 * @param Sender The message sender (may be nullptr if the sender has no subscriptions).
 	 * @param OutRecipients Will hold the collection of recipients.
 	 */
-	void FilterSubscriptions(TArray<IMessageSubscriptionPtr>& Subscriptions, const IMessageContextRef& Context, TArray<IReceiveMessagesPtr>& OutRecipients);
+	void FilterSubscriptions(
+		TArray<TSharedPtr<IMessageSubscription, ESPMode::ThreadSafe>>& Subscriptions,
+		const TSharedRef<IMessageContext, ESPMode::ThreadSafe>& Context,
+		TArray<TSharedPtr<IMessageReceiver, ESPMode::ThreadSafe>>& OutRecipients);
 
 	/**
 	 * Dispatches a single message to its recipients.
 	 *
 	 * @param Message The message to dispatch.
 	 */
-	void DispatchMessage(const IMessageContextRef& Message);
+	void DispatchMessage(const TSharedRef<IMessageContext, ESPMode::ThreadSafe>& Message);
 
 	/** Processes all delayed messages. */
 	void ProcessDelayedMessages();
 
 private:
 
-	// Structure for delayed messages.
+	/** Structure for delayed messages. */
 	struct FDelayedMessage
 	{
-		// Holds the context of the delayed message.
+		/** Holds the context of the delayed message. */
 		IMessageContextPtr Context;
 
-		// Holds a sequence number used by the delayed message queue.
+		/** Holds a sequence number used by the delayed message queue. */
 		int64 Sequence;
 
 
-		// Default constructor.
+		/** Default constructor. */
 		FDelayedMessage() { }
 
-		// Creates and initializes a new instance.
-		FDelayedMessage(const IMessageContextRef& InContext, int64 InSequence)
+		/** Creates and initializes a new instance. */
+		FDelayedMessage(const TSharedRef<IMessageContext, ESPMode::ThreadSafe>& InContext, int64 InSequence)
 			: Context(InContext)
 			, Sequence(InSequence)
 		{ }
 
-		// Comparison operator for heap sorting.
+		/** Comparison operator for heap sorting. */
 		bool operator<(const FDelayedMessage& Other) const
 		{
 			const FTimespan Difference = Other.Context->GetTimeSent() - Context->GetTimeSent();
@@ -207,36 +211,36 @@ private:
 private:
 
 	/** Handles adding message interceptors. */
-	void HandleAddInterceptor(IMessageInterceptorRef Interceptor, FName MessageType);
+	void HandleAddInterceptor(TSharedRef<IMessageInterceptor, ESPMode::ThreadSafe> Interceptor, FName MessageType);
 
 	/** Handles adding message recipients. */
-	void HandleAddRecipient(FMessageAddress Address, IReceiveMessagesWeakPtr RecipientPtr);
+	void HandleAddRecipient(FMessageAddress Address, TWeakPtr<IMessageReceiver, ESPMode::ThreadSafe> RecipientPtr);
 
 	/** Handles adding of subscriptions. */
-	void HandleAddSubscriber(IMessageSubscriptionRef Subscription);
+	void HandleAddSubscriber(TSharedRef<IMessageSubscription, ESPMode::ThreadSafe> Subscription);
 
 	/** Handles the removal of message interceptors. */
-	void HandleRemoveInterceptor(IMessageInterceptorRef Interceptor, FName MessageType);
+	void HandleRemoveInterceptor(TSharedRef<IMessageInterceptor, ESPMode::ThreadSafe> Interceptor, FName MessageType);
 
 	/** Handles the removal of message recipients. */
 	void HandleRemoveRecipient(FMessageAddress Address);
 
 	/** Handles the removal of subscribers. */
-	void HandleRemoveSubscriber(IReceiveMessagesWeakPtr SubscriberPtr, FName MessageType);
+	void HandleRemoveSubscriber(TWeakPtr<IMessageReceiver, ESPMode::ThreadSafe> SubscriberPtr, FName MessageType);
 
 	/** Handles the routing of messages. */
-	void HandleRouteMessage(IMessageContextRef Context);
+	void HandleRouteMessage(TSharedRef<IMessageContext, ESPMode::ThreadSafe> Context);
 
 private:
 
 	/** Maps message types to interceptors. */
-	TMap<FName, TArray<IMessageInterceptorPtr>> ActiveInterceptors;
+	TMap<FName, TArray<TSharedPtr<IMessageInterceptor, ESPMode::ThreadSafe>>> ActiveInterceptors;
 
 	/** Maps message addresses to recipients. */
-	TMap<FMessageAddress, IReceiveMessagesWeakPtr> ActiveRecipients;
+	TMap<FMessageAddress, TWeakPtr<IMessageReceiver, ESPMode::ThreadSafe>> ActiveRecipients;
 
 	/** Maps message types to subscriptions. */
-	TMap<FName, TArray<IMessageSubscriptionPtr>> ActiveSubscriptions;
+	TMap<FName, TArray<TSharedPtr<IMessageSubscription, ESPMode::ThreadSafe>>> ActiveSubscriptions;
 
 	/** Holds the router command queue. */
 	TQueue<CommandDelegate, EQueueMode::Mpsc> Commands;
@@ -254,7 +258,7 @@ private:
 	bool Stopping;
 
 	/** Holds the message tracer. */
-	FMessageTracerRef Tracer;
+	TSharedRef<FMessageTracer, ESPMode::ThreadSafe> Tracer;
 
 	/** Holds an event signaling that work is available. */
 	FEvent* WorkEvent;
diff --git a/Engine/Source/Runtime/Messaging/Private/Bus/MessageSubscription.h b/Engine/Source/Runtime/Messaging/Private/Bus/MessageSubscription.h
--- a/Engine/Source/Runtime/Messaging/Private/Bus/MessageSubscription.h
+++ b/Engine/Source/Runtime/Messaging/Private/Bus/MessageSubscription.h
@@ -4,7 +4,9 @@
 
 #include "IMessageContext.h"
 #include "IMessageSubscription.h"
-#include "IReceiveMessages.h"
+
+
+class IMessageReceiver;
 
 
 /**
@@ -27,7 +29,7 @@ public:
 	 * @param InReceivingThread The thread on which to receive messages on.
 	 * @param InScopeRange The message scope range to subscribe to.
 	 */
-	FMessageSubscription(const IReceiveMessagesRef& InSubscriber, const FName& InMessageType, const FMessageScopeRange& InScopeRange)
+	FMessageSubscription(const TSharedRef<IMessageReceiver, ESPMode::ThreadSafe>& InSubscriber, const FName& InMessageType, const FMessageScopeRange& InScopeRange)
 		: Enabled(true)
 		, MessageType(InMessageType)
 		, ScopeRange(InScopeRange)
@@ -36,7 +38,7 @@ public:
 
 public:
 
-	// IMessageSubscription interface
+	//~ IMessageSubscription interface
 
 	virtual void Disable() override
 	{
@@ -58,7 +60,7 @@ public:
 		return ScopeRange;
 	}
 
-	virtual const IReceiveMessagesWeakPtr& GetSubscriber() override
+	virtual const TWeakPtr<IMessageReceiver, ESPMode::ThreadSafe>& GetSubscriber() override
 	{
 		return Subscriber;
 	}
@@ -80,5 +82,5 @@ private:
 	FMessageScopeRange ScopeRange;
 
 	/** Holds the subscriber. */
-	IReceiveMessagesWeakPtr Subscriber;
+	TWeakPtr<IMessageReceiver, ESPMode::ThreadSafe> Subscriber;
 };
diff --git a/Engine/Source/Runtime/Messaging/Private/Bus/MessageTracer.cpp b/Engine/Source/Runtime/Messaging/Private/Bus/MessageTracer.cpp
--- a/Engine/Source/Runtime/Messaging/Private/Bus/MessageTracer.cpp
+++ b/Engine/Source/Runtime/Messaging/Private/Bus/MessageTracer.cpp
@@ -1,6 +1,7 @@
 // Copyright 1998-2016 Epic Games, Inc. All Rights Reserved.
 
 #include "MessagingPrivatePCH.h"
+#include "MessageTracer.h"
 
 
 /* FMessageTracer structors
@@ -27,13 +28,13 @@ FMessageTracer::~FMessageTracer()
 /* FMessageTracer interface
  *****************************************************************************/
 
-void FMessageTracer::TraceAddedInterceptor(const IMessageInterceptorRef& Interceptor, const FName& MessageType)
+void FMessageTracer::TraceAddedInterceptor(const TSharedRef<IMessageInterceptor, ESPMode::ThreadSafe>& Interceptor, const FName& MessageType)
 {
 	double Timestamp = FPlatformTime::Seconds();
 
 	Traces.Enqueue([=]() {
 		// create interceptor information
-		FMessageTracerInterceptorInfoPtr& InterceptorInfo = Interceptors.FindOrAdd(Interceptor->GetInterceptorId());
+		auto& InterceptorInfo = Interceptors.FindOrAdd(Interceptor->GetInterceptorId());
 
 		if (!InterceptorInfo.IsValid())
 		{
@@ -48,13 +49,13 @@ void FMessageTracer::TraceAddedInterceptor(const IMessageInterceptorRef& Interce
 }
 
 
-void FMessageTracer::TraceAddedRecipient(const FMessageAddress& Address, const IReceiveMessagesRef& Recipient)
+void FMessageTracer::TraceAddedRecipient(const FMessageAddress& Address, const TSharedRef<IMessageReceiver, ESPMode::ThreadSafe>& Recipient)
 {
 	double Timestamp = FPlatformTime::Seconds();
 
 	Traces.Enqueue([=]() {
 		// create endpoint information
-		FMessageTracerEndpointInfoPtr& EndpointInfo = RecipientsToEndpointInfos.FindOrAdd(Recipient->GetRecipientId());
+		TSharedPtr<FMessageTracerEndpointInfo>& EndpointInfo = RecipientsToEndpointInfos.FindOrAdd(Recipient->GetRecipientId());
 
 		if (!EndpointInfo.IsValid())
 		{
@@ -62,7 +63,7 @@ void FMessageTracer::TraceAddedRecipient(const FMessageAddress& Address, const I
 		}
 
 		// initialize endpoint information
-		FMessageTracerAddressInfoRef AddressInfo = MakeShareable(new FMessageTracerAddressInfo());
+		TSharedRef<FMessageTracerAddressInfo> AddressInfo = MakeShareable(new FMessageTracerAddressInfo());
 		{
 			AddressInfo->Address = Address;
 			AddressInfo->TimeRegistered = Timestamp;
@@ -79,7 +80,7 @@ void FMessageTracer::TraceAddedRecipient(const FMessageAddress& Address, const I
 }
 
 
-void FMessageTracer::TraceAddedSubscription(const IMessageSubscriptionRef& Subscription)
+void FMessageTracer::TraceAddedSubscription(const TSharedRef<IMessageSubscription, ESPMode::ThreadSafe>& Subscription)
 {
 	if (!Running)
 	{
@@ -94,7 +95,7 @@ void FMessageTracer::TraceAddedSubscription(const IMessageSubscriptionRef& Subsc
 }
 
 
-void FMessageTracer::TraceDispatchedMessage(const IMessageContextRef& Context, const IReceiveMessagesRef& Recipient, bool Async)
+void FMessageTracer::TraceDispatchedMessage(const IMessageContextRef& Context, const TSharedRef<IMessageReceiver, ESPMode::ThreadSafe>& Recipient, bool Async)
 {
 	if (!Running)
 	{
@@ -105,14 +106,14 @@ void FMessageTracer::TraceDispatchedMessage(const IMessageContextRef& Context, c
 
 	Traces.Enqueue([=]() {
 		// look up message & endpoint info
-		FMessageTracerMessageInfoPtr MessageInfo = MessageInfos.FindRef(Context);
+		TSharedPtr<FMessageTracerMessageInfo> MessageInfo = MessageInfos.FindRef(Context);
 
 		if (!MessageInfo.IsValid())
 		{
 			return;
 		}
 
-		FMessageTracerEndpointInfoPtr& EndpointInfo = RecipientsToEndpointInfos.FindOrAdd(Recipient->GetRecipientId());
+		TSharedPtr<FMessageTracerEndpointInfo>& EndpointInfo = RecipientsToEndpointInfos.FindOrAdd(Recipient->GetRecipientId());
 
 		if (!EndpointInfo.IsValid())
 		{
@@ -120,7 +121,7 @@ void FMessageTracer::TraceDispatchedMessage(const IMessageContextRef& Context, c
 		}
 
 		// update message information
-		FMessageTracerDispatchStateRef DispatchState = MakeShareable(new FMessageTracerDispatchState());
+		TSharedRef<FMessageTracerDispatchState> DispatchState = MakeShareable(new FMessageTracerDispatchState());
 		{
 			DispatchState->DispatchLatency = Timestamp - MessageInfo->TimeSent;
 			DispatchState->DispatchType = Async ? EMessageTracerDispatchTypes::TaskGraph : EMessageTracerDispatchTypes::Direct;
@@ -138,7 +139,7 @@ void FMessageTracer::TraceDispatchedMessage(const IMessageContextRef& Context, c
 }
 
 
-void FMessageTracer::TraceHandledMessage(const IMessageContextRef& Context, const IReceiveMessagesRef& Recipient)
+void FMessageTracer::TraceHandledMessage(const IMessageContextRef& Context, const TSharedRef<IMessageReceiver, ESPMode::ThreadSafe>& Recipient)
 {
 	if (!Running)
 	{
@@ -149,14 +150,14 @@ void FMessageTracer::TraceHandledMessage(const IMessageContextRef& Context, cons
 
 	Traces.Enqueue([=]() {
 		// look up message & endpoint info
-		FMessageTracerMessageInfoPtr MessageInfo = MessageInfos.FindRef(Context);
+		TSharedPtr<FMessageTracerMessageInfo> MessageInfo = MessageInfos.FindRef(Context);
 
 		if (!MessageInfo.IsValid())
 		{
 			return;
 		}
 
-		FMessageTracerEndpointInfoPtr EndpointInfo = RecipientsToEndpointInfos.FindRef(Recipient->GetRecipientId());
+		TSharedPtr<FMessageTracerEndpointInfo> EndpointInfo = RecipientsToEndpointInfos.FindRef(Recipient->GetRecipientId());
 
 		if (!EndpointInfo.IsValid())
 		{
@@ -164,7 +165,7 @@ void FMessageTracer::TraceHandledMessage(const IMessageContextRef& Context, cons
 		}
 
 		// update message information
-		FMessageTracerDispatchStatePtr DispatchState = MessageInfo->DispatchStates.FindRef(EndpointInfo);
+		TSharedPtr<FMessageTracerDispatchState> DispatchState = MessageInfo->DispatchStates.FindRef(EndpointInfo);
 
 		if (DispatchState.IsValid())
 		{
@@ -174,7 +175,7 @@ void FMessageTracer::TraceHandledMessage(const IMessageContextRef& Context, cons
 }
 
 
-void FMessageTracer::TraceInterceptedMessage(const IMessageContextRef& Context, const IMessageInterceptorRef& Interceptor)
+void FMessageTracer::TraceInterceptedMessage(const IMessageContextRef& Context, const TSharedRef<IMessageInterceptor, ESPMode::ThreadSafe>& Interceptor)
 {
 	if (!Running)
 	{
@@ -185,7 +186,7 @@ void FMessageTracer::TraceInterceptedMessage(const IMessageContextRef& Context,
 
 	Traces.Enqueue([=]() {
 		// look up message & interceptor info
-		FMessageTracerMessageInfoPtr MessageInfo = MessageInfos.FindRef(Context);
+		auto MessageInfo = MessageInfos.FindRef(Context);
 
 		if (!MessageInfo.IsValid())
 		{
@@ -194,7 +195,7 @@ void FMessageTracer::TraceInterceptedMessage(const IMessageContextRef& Context,
 
 		MessageInfo->Intercepted = true;
 
-		FMessageTracerInterceptorInfoPtr InterceptorInfo = Interceptors.FindRef(Interceptor->GetInterceptorId());
+		auto InterceptorInfo = Interceptors.FindRef(Interceptor->GetInterceptorId());
 
 		if (!InterceptorInfo.IsValid())
 		{
@@ -207,12 +208,12 @@ void FMessageTracer::TraceInterceptedMessage(const IMessageContextRef& Context,
 }
 
 
-void FMessageTracer::TraceRemovedInterceptor(const IMessageInterceptorRef& Interceptor, const FName& MessageType)
+void FMessageTracer::TraceRemovedInterceptor(const TSharedRef<IMessageInterceptor, ESPMode::ThreadSafe>& Interceptor, const FName& MessageType)
 {
 	double Timestamp = FPlatformTime::Seconds();
 
 	Traces.Enqueue([=]() {
-		FMessageTracerInterceptorInfoPtr InterceptorInfo = Interceptors.FindRef(Interceptor->GetInterceptorId());
+		auto InterceptorInfo = Interceptors.FindRef(Interceptor->GetInterceptorId());
 
 		if (!InterceptorInfo.IsValid())
 		{
@@ -230,7 +231,7 @@ void FMessageTracer::TraceRemovedRecipient(const FMessageAddress& Address)
 	double Timestamp = FPlatformTime::Seconds();
 
 	Traces.Enqueue([=]() {
-		FMessageTracerEndpointInfoPtr EndpointInfo = AddressesToEndpointInfos.FindRef(Address);
+		TSharedPtr<FMessageTracerEndpointInfo> EndpointInfo = AddressesToEndpointInfos.FindRef(Address);
 
 		if (!EndpointInfo.IsValid())
 		{
@@ -238,7 +239,7 @@ void FMessageTracer::TraceRemovedRecipient(const FMessageAddress& Address)
 		}
 
 		// update endpoint information
-		FMessageTracerAddressInfoPtr AddressInfo = EndpointInfo->AddressInfos.FindRef(Address);
+		TSharedPtr<FMessageTracerAddressInfo> AddressInfo = EndpointInfo->AddressInfos.FindRef(Address);
 
 		if (AddressInfo.IsValid())
 		{
@@ -248,7 +249,7 @@ void FMessageTracer::TraceRemovedRecipient(const FMessageAddress& Address)
 }
 
 
-void FMessageTracer::TraceRemovedSubscription(const IMessageSubscriptionRef& Subscription, const FName& MessageType)
+void FMessageTracer::TraceRemovedSubscription(const TSharedRef<IMessageSubscription, ESPMode::ThreadSafe>& Subscription, const FName& MessageType)
 {
 	if (!Running)
 	{
@@ -280,7 +281,7 @@ void FMessageTracer::TraceRoutedMessage(const IMessageContextRef& Context)
 
 	Traces.Enqueue([=]() {
 		// update message information
-		FMessageTracerMessageInfoPtr MessageInfo = MessageInfos.FindRef(Context);
+		TSharedPtr<FMessageTracerMessageInfo> MessageInfo = MessageInfos.FindRef(Context);
 
 		if (MessageInfo.IsValid())
 		{
@@ -301,7 +302,7 @@ void FMessageTracer::TraceSentMessage(const IMessageContextRef& Context)
 
 	Traces.Enqueue([=]() {
 		// look up endpoint info
-		FMessageTracerEndpointInfoPtr EndpointInfo = AddressesToEndpointInfos.FindRef(Context->GetSender());
+		TSharedPtr<FMessageTracerEndpointInfo> EndpointInfo = AddressesToEndpointInfos.FindRef(Context->GetSender());
 
 		if (!EndpointInfo.IsValid())
 		{
@@ -309,7 +310,7 @@ void FMessageTracer::TraceSentMessage(const IMessageContextRef& Context)
 		}
 
 		// create message info
-		FMessageTracerMessageInfoRef MessageInfo = MakeShareable(new FMessageTracerMessageInfo());
+		TSharedRef<FMessageTracerMessageInfo> MessageInfo = MakeShareable(new FMessageTracerMessageInfo());
 		{
 			MessageInfo->Context = Context;
 			MessageInfo->Intercepted = false;
@@ -320,7 +321,7 @@ void FMessageTracer::TraceSentMessage(const IMessageContextRef& Context)
 		}
 
 		// add message type
-		FMessageTracerTypeInfoPtr& TypeInfo = MessageTypes.FindOrAdd(Context->GetMessageType());
+		TSharedPtr<FMessageTracerTypeInfo>& TypeInfo = MessageTypes.FindOrAdd(Context->GetMessageType());
 
 		if (!TypeInfo.IsValid())
 		{
@@ -344,7 +345,27 @@ void FMessageTracer::TraceSentMessage(const IMessageContextRef& Context)
 /* IMessageTracer interface
  *****************************************************************************/
 
-int32 FMessageTracer::GetEndpoints(TArray<FMessageTracerEndpointInfoPtr>& OutEndpoints) const
+void FMessageTracer::Break()
+{
+	Breaking = true;
+}
+
+
+void FMessageTracer::Continue()
+{
+	if (!Running)
+	{
+		Running = true;
+	}
+	else if (Breaking)
+	{
+		Breaking = false;
+		ContinueEvent->Trigger();
+	}
+}
+
+
+int32 FMessageTracer::GetEndpoints(TArray<TSharedPtr<FMessageTracerEndpointInfo>>& OutEndpoints) const
 {
 	RecipientsToEndpointInfos.GenerateValueArray(OutEndpoints);
 
@@ -352,7 +373,7 @@ int32 FMessageTracer::GetEndpoints(TArray<FMessageTracerEndpointInfoPtr>& OutEnd
 }
 
 
-int32 FMessageTracer::GetMessages(TArray<FMessageTracerMessageInfoPtr>& OutMessages) const
+int32 FMessageTracer::GetMessages(TArray<TSharedPtr<FMessageTracerMessageInfo>>& OutMessages) const
 {
 	MessageInfos.GenerateValueArray(OutMessages);
 
@@ -360,7 +381,7 @@ int32 FMessageTracer::GetMessages(TArray<FMessageTracerMessageInfoPtr>& OutMessa
 }
 
 
-int32 FMessageTracer::GetMessageTypes(TArray<FMessageTracerTypeInfoPtr>& OutTypes) const
+int32 FMessageTracer::GetMessageTypes(TArray<TSharedPtr<FMessageTracerTypeInfo>>& OutTypes) const
 {
 	MessageTypes.GenerateValueArray(OutTypes);
 
@@ -368,6 +389,58 @@ int32 FMessageTracer::GetMessageTypes(TArray<FMessageTracerTypeInfoPtr>& OutType
 }
 
 
+bool FMessageTracer::HasMessages() const
+{
+	return (MessageInfos.Num() > 0);
+}
+
+
+bool FMessageTracer::IsBreaking() const
+{
+	return Breaking;
+}
+
+
+bool FMessageTracer::IsRunning() const
+{
+	return Running;
+}
+
+
+void FMessageTracer::Reset()
+{
+	ResetPending = true;
+}
+
+
+void FMessageTracer::Step()
+{
+	if (!Breaking)
+	{
+		return;
+	}
+
+	ContinueEvent->Trigger();
+}
+
+
+void FMessageTracer::Stop()
+{
+	if (!Running)
+	{
+		return;
+	}
+
+	Running = false;
+
+	if (Breaking)
+	{
+		Breaking = false;
+		ContinueEvent->Trigger();
+	}
+}
+
+
 bool FMessageTracer::Tick(float DeltaTime)
 {
 	if (ResetPending)
@@ -401,7 +474,7 @@ void FMessageTracer::ResetMessages()
 
 	for (auto& EndpointInfoPair : AddressesToEndpointInfos)
 	{
-		FMessageTracerEndpointInfoPtr& EndpointInfo = EndpointInfoPair.Value;
+		TSharedPtr<FMessageTracerEndpointInfo>& EndpointInfo = EndpointInfoPair.Value;
 		{
 			EndpointInfo->ReceivedMessages.Reset();
 			EndpointInfo->SentMessages.Reset();
diff --git a/Engine/Source/Runtime/Messaging/Private/Bus/MessageTracer.h b/Engine/Source/Runtime/Messaging/Private/Bus/MessageTracer.h
--- a/Engine/Source/Runtime/Messaging/Private/Bus/MessageTracer.h
+++ b/Engine/Source/Runtime/Messaging/Private/Bus/MessageTracer.h
@@ -3,11 +3,18 @@
 #pragma once
 
 #include "IMessageContext.h"
-#include "IMessageInterceptor.h"
-#include "IMessageSubscription.h"
 #include "IMessageTracer.h"
-#include "IMessageTracerBreakpoint.h"
-#include "IReceiveMessages.h"
+
+
+struct FMessageTracerEndpointInfo;
+struct FMessageTracerInterceptorInfo;
+struct FMessageTracerMessageInfo;
+struct FMessageTracerTypeInfo;
+
+class IMessageInterceptor;
+class IMessageReceiver;
+class IMessageSubscription;
+class IMessageTracerBreakpoint;
 
 
 /**
@@ -32,7 +39,7 @@ public:
 	 * @param Interceptor The added interceptor.
 	 * @param MessageType The type of messages being intercepted.
 	 */
-	void TraceAddedInterceptor(const IMessageInterceptorRef& Interceptor, const FName& MessageType);
+	void TraceAddedInterceptor(const TSharedRef<IMessageInterceptor, ESPMode::ThreadSafe>& Interceptor, const FName& MessageType);
 
 	/**
 	 * Notifies the tracer that a message recipient has been added to the message bus.
@@ -40,14 +47,14 @@ public:
 	 * @param Address The address of the added recipient.
 	 * @param Recipient The added recipient.
 	 */
-	void TraceAddedRecipient(const FMessageAddress& Address, const IReceiveMessagesRef& Recipient);
+	void TraceAddedRecipient(const FMessageAddress& Address, const TSharedRef<IMessageReceiver, ESPMode::ThreadSafe>& Recipient);
 
 	/**
 	 * Notifies the tracer that a message subscription has been added to the message bus.
 	 *
 	 * @param Subscription The added subscription.
 	 */
-	void TraceAddedSubscription(const IMessageSubscriptionRef& Subscription);
+	void TraceAddedSubscription(const TSharedRef<IMessageSubscription, ESPMode::ThreadSafe>& Subscription);
 
 	/**
 	 * Notifies the tracer that a message has been dispatched.
@@ -56,7 +63,7 @@ public:
 	 * @param Recipient The message recipient.
 	 * @param Async Whether the message was dispatched asynchronously.
 	 */
-	void TraceDispatchedMessage(const IMessageContextRef& Context, const IReceiveMessagesRef& Recipient, bool Async);
+	void TraceDispatchedMessage(const TSharedRef<IMessageContext, ESPMode::ThreadSafe>& Context, const TSharedRef<IMessageReceiver, ESPMode::ThreadSafe>& Recipient, bool Async);
 
 	/**
 	 * Notifies the tracer that a message has been handled.
@@ -64,7 +71,7 @@ public:
 	 * @param Context The context of the dispatched message.
 	 * @param Recipient The message recipient that handled the message.
 	 */
-	void TraceHandledMessage(const IMessageContextRef& Context, const IReceiveMessagesRef& Recipient);
+	void TraceHandledMessage(const TSharedRef<IMessageContext, ESPMode::ThreadSafe>& Context, const TSharedRef<IMessageReceiver, ESPMode::ThreadSafe>& Recipient);
 
 	/**
 	 * Notifies the tracer that a message has been intercepted.
@@ -72,7 +79,7 @@ public:
 	 * @param Context The context of the intercepted message.
 	 * @param Interceptor The interceptor.
 	 */
-	void TraceInterceptedMessage(const IMessageContextRef& Context, const IMessageInterceptorRef& Interceptor);
+	void TraceInterceptedMessage(const TSharedRef<IMessageContext, ESPMode::ThreadSafe>& Context, const TSharedRef<IMessageInterceptor, ESPMode::ThreadSafe>& Interceptor);
 
 	/**
 	 * Notifies the tracer that a message interceptor has been removed from the message bus.
@@ -80,7 +87,7 @@ public:
 	 * @param Interceptor The removed interceptor.
 	 * @param MessageType The type of messages that is no longer being intercepted.
 	 */
-	void TraceRemovedInterceptor(const IMessageInterceptorRef& Interceptor, const FName& MessageType);
+	void TraceRemovedInterceptor(const TSharedRef<IMessageInterceptor, ESPMode::ThreadSafe>& Interceptor, const FName& MessageType);
 
 	/**
 	 * Notifies the tracer that a recipient has been removed from the message bus.
@@ -95,62 +102,34 @@ public:
 	 * @param Subscriber The removed subscriber.
 	 * @param MessageType The type of messages no longer being subscribed to.
 	 */
-	void TraceRemovedSubscription(const IMessageSubscriptionRef& Subscription, const FName& MessageType);
+	void TraceRemovedSubscription(const TSharedRef<IMessageSubscription, ESPMode::ThreadSafe>& Subscription, const FName& MessageType);
 
 	/**
 	 * Notifies the tracer that a message has been routed.
 	 *
 	 * @param Context The context of the routed message.
 	 */
-	void TraceRoutedMessage(const IMessageContextRef& Context);
+	void TraceRoutedMessage(const TSharedRef<IMessageContext, ESPMode::ThreadSafe>& Context);
 
 	/**
 	 * Notifies the tracer that a message has been sent.
 	 *
 	 * @param Context The context of the sent message.
 	 */
-	void TraceSentMessage(const IMessageContextRef& Context);
+	void TraceSentMessage(const TSharedRef<IMessageContext, ESPMode::ThreadSafe>& Context);
 
 public:
 
-	// IMessageTracer interface
-
-	virtual void Break() override
-	{
-		Breaking = true;
-	}
-
-	virtual void Continue() override
-	{
-		if (!Running)
-		{
-			Running = true;
-		}
-		else if (Breaking)
-		{
-			Breaking = false;
-			ContinueEvent->Trigger();
-		}
-	}
-
-	virtual int32 GetEndpoints(TArray<FMessageTracerEndpointInfoPtr>& OutEndpoints) const override;
-	virtual int32 GetMessages(TArray<FMessageTracerMessageInfoPtr>& OutMessages) const override;
-	virtual int32 GetMessageTypes(TArray<FMessageTracerTypeInfoPtr>& OutTypes) const override;
-
-	virtual bool HasMessages() const override
-	{
-		return (MessageInfos.Num() > 0);
-	}
-
-	virtual bool IsBreaking() const override
-	{
-		return Breaking;
-	}
+	//~ IMessageTracer interface
 
-	virtual bool IsRunning() const override
-	{
-		return Running;
-	}
+	virtual void Break() override;
+	virtual void Continue() override;
+	virtual int32 GetEndpoints(TArray<TSharedPtr<FMessageTracerEndpointInfo>>& OutEndpoints) const override;
+	virtual int32 GetMessages(TArray<TSharedPtr<FMessageTracerMessageInfo>>& OutMessages) const override;
+	virtual int32 GetMessageTypes(TArray<TSharedPtr<FMessageTracerTypeInfo>>& OutTypes) const override;
+	virtual bool HasMessages() const override;
+	virtual bool IsBreaking() const override;
+	virtual bool IsRunning() const override;
 
 	DECLARE_DERIVED_EVENT(FMessageTracer, IMessageTracer::FOnMessageAdded, FOnMessageAdded)
 	virtual FOnMessageAdded& OnMessageAdded() override
@@ -170,37 +149,9 @@ public:
 		return TypeAddedDelegate;
 	}
 
-	virtual void Reset() override
-	{
-		ResetPending = true;
-	}
-
-	virtual void Step() override
-	{
-		if (!Breaking)
-		{
-			return;
-		}
-
-		ContinueEvent->Trigger();
-	}
-
-	virtual void Stop() override
-	{
-		if (!Running)
-		{
-			return;
-		}
-
-		Running = false;
-
-		if (Breaking)
-		{
-			Breaking = false;
-			ContinueEvent->Trigger();
-		}
-	}
-
+	virtual void Reset() override;
+	virtual void Step() override;
+	virtual void Stop();
 	virtual bool Tick(float DeltaTime) override;
 
 protected:
@@ -213,33 +164,33 @@ protected:
 	 *
 	 * @param Context The context of the message to consider for breaking.
 	 */
-	bool ShouldBreak(const IMessageContextRef& Context) const;
+	bool ShouldBreak(const TSharedRef<IMessageContext, ESPMode::ThreadSafe>& Context) const;
 
 private:
 
 	/** Holds the collection of endpoints for known message addresses. */
-	TMap<FMessageAddress, FMessageTracerEndpointInfoPtr> AddressesToEndpointInfos;
+	TMap<FMessageAddress, TSharedPtr<FMessageTracerEndpointInfo>> AddressesToEndpointInfos;
 
 	/** Holds a flag indicating whether a breakpoint was hit. */
 	bool Breaking;
 
 	/** Holds the collection of breakpoints. */
-	TArray<IMessageTracerBreakpointPtr> Breakpoints;
+	TArray<TSharedPtr<IMessageTracerBreakpoint, ESPMode::ThreadSafe>> Breakpoints;
 
 	/** Holds an event signaling that messaging routing can continue. */
 	FEvent* ContinueEvent;
 
 	/** The collection of known interceptors. */
-	TMap<FGuid, FMessageTracerInterceptorInfoPtr> Interceptors;
+	TMap<FGuid, TSharedPtr<FMessageTracerInterceptorInfo>> Interceptors;
 
 	/** Holds the collection of endpoints for known recipient identifiers. */
-	TMap<FGuid, FMessageTracerEndpointInfoPtr> RecipientsToEndpointInfos;
+	TMap<FGuid, TSharedPtr<FMessageTracerEndpointInfo>> RecipientsToEndpointInfos;
 
 	/** Holds the collection of known messages. */
-	TMap<IMessageContextPtr, FMessageTracerMessageInfoPtr> MessageInfos;
+	TMap<TSharedPtr<IMessageContext, ESPMode::ThreadSafe>, TSharedPtr<FMessageTracerMessageInfo>> MessageInfos;
 
 	/** Holds the collection of known message types. */
-	TMap<FName, FMessageTracerTypeInfoPtr> MessageTypes;
+	TMap<FName, TSharedPtr<FMessageTracerTypeInfo>> MessageTypes;
 
 	/** Holds a flag indicating whether a reset is pending. */
 	bool ResetPending;
@@ -264,13 +215,3 @@ private:
 	/** Holds a delegate that is executed when a new type has been added to the collection of known message types. */
 	FOnTypeAdded TypeAddedDelegate;
 };
-
-
-/** Type definition for weak pointers to instances of FMessageTracer. */
-typedef TWeakPtr<FMessageTracer, ESPMode::ThreadSafe> FMessageTracerWeakPtr;
-
-/** Type definition for shared pointers to instances of FMessageTracer. */
-typedef TSharedPtr<FMessageTracer, ESPMode::ThreadSafe> FMessageTracerPtr;
-
-/** Type definition for shared references to instances of FMessageTracer. */
-typedef TSharedRef<FMessageTracer, ESPMode::ThreadSafe> FMessageTracerRef;
diff --git a/Engine/Source/Runtime/Messaging/Private/MessagingModule.cpp b/Engine/Source/Runtime/Messaging/Private/MessagingModule.cpp
--- a/Engine/Source/Runtime/Messaging/Private/MessagingModule.cpp
+++ b/Engine/Source/Runtime/Messaging/Private/MessagingModule.cpp
@@ -2,7 +2,8 @@
 
 #include "MessagingPrivatePCH.h"
 #include "IMessagingModule.h"
-#include "ModuleManager.h"
+#include "MessageBridge.h"
+#include "MessageBus.h"
 
 
 #ifndef PLATFORM_SUPPORTS_MESSAGEBUS
@@ -19,9 +20,9 @@ class FMessagingModule
 {
 public:
 
-	// FSelfRegisteringExec interface
+	//~ FSelfRegisteringExec interface
 
-	virtual bool Exec( UWorld* InWorld, const TCHAR* Cmd, FOutputDevice& Ar ) override
+	virtual bool Exec(UWorld* InWorld, const TCHAR* Cmd, FOutputDevice& Ar) override
 	{
 		if (!FParse::Command(&Cmd, TEXT("MESSAGING")))
 		{
@@ -53,26 +54,26 @@ public:
 
 public:
 
-	// IMessagingModule interface
+	//~ IMessagingModule interface
 
-	virtual IMessageBridgePtr CreateBridge( const FMessageAddress& Address, const IMessageBusRef& Bus, const IMessageTransportRef& Transport ) override
+	virtual TSharedPtr<IMessageBridge, ESPMode::ThreadSafe> CreateBridge(const FMessageAddress& Address, const TSharedRef<IMessageBus, ESPMode::ThreadSafe>& Bus, const IMessageTransportRef& Transport) override
 	{
 		return MakeShareable(new FMessageBridge(Address, Bus, Transport));
 	}
 
-	virtual IMessageBusPtr CreateBus( const IAuthorizeMessageRecipientsPtr& RecipientAuthorizer ) override
+	virtual TSharedPtr<IMessageBus, ESPMode::ThreadSafe> CreateBus(const TSharedPtr<IAuthorizeMessageRecipients>& RecipientAuthorizer) override
 	{
 		return MakeShareable(new FMessageBus(RecipientAuthorizer));
 	}
 
-	virtual IMessageBusPtr GetDefaultBus() const override
+	virtual TSharedPtr<IMessageBus, ESPMode::ThreadSafe> GetDefaultBus() const override
 	{
 		return DefaultBus;
 	}
 
 public:
 
-	// IModuleInterface interface
+	//~ IModuleInterface interface
 
 	virtual void StartupModule() override
 	{
@@ -101,7 +102,7 @@ protected:
 			return;
 		}
 
-		IMessageBusWeakPtr DefaultBusPtr = DefaultBus;
+		TWeakPtr<IMessageBus, ESPMode::ThreadSafe> DefaultBusPtr = DefaultBus;
 
 		DefaultBus->Shutdown();
 		DefaultBus.Reset();
@@ -120,7 +121,7 @@ protected:
 
 private:
 
-	// Callback for Core shutdown.
+	/** Callback for Core shutdown. */
 	void HandleCorePreExit()
 	{
 		ShutdownDefaultBus();
@@ -128,18 +129,9 @@ private:
 
 private:
 
-	// Holds the message bus.
-	IMessageBusPtr DefaultBus;
+	/** Holds the message bus. */
+	TSharedPtr<IMessageBus, ESPMode::ThreadSafe> DefaultBus;
 };
 
 
-/* Misc static functions
- *****************************************************************************/
-
-TStatId FMessageDispatchTask::GetStatId() const
-{
-	RETURN_QUICK_DECLARE_CYCLE_STAT(FMessageDispatchTask, STATGROUP_TaskGraphTasks);
-}
-
-
 IMPLEMENT_MODULE(FMessagingModule, Messaging);
diff --git a/Engine/Source/Runtime/Messaging/Private/MessagingPrivatePCH.h b/Engine/Source/Runtime/Messaging/Private/MessagingPrivatePCH.h
--- a/Engine/Source/Runtime/Messaging/Private/MessagingPrivatePCH.h
+++ b/Engine/Source/Runtime/Messaging/Private/MessagingPrivatePCH.h
@@ -2,22 +2,21 @@
 
 #pragma once
 
+#include "Runtime/Core/Public/Core.h"
+#include "Runtime/Core/Public/Modules/ModuleManager.h"
+#include "Runtime/CoreUObject/Public/CoreUObject.h"
 
-/* Private dependencies
- *****************************************************************************/
-
-#include "CoreUObject.h"
-#include "TaskGraphInterfaces.h"
-
-
-/* Private includes
- *****************************************************************************/
-
-#include "MessageContext.h"
-#include "MessageSubscription.h"
-#include "MessageTracer.h"
-#include "MessageDispatchTask.h"
-#include "MessageRouter.h"
-#include "MessageBus.h"
-#include "MessageAddressBook.h"
-#include "MessageBridge.h"
+#include "IAuthorizeMessageRecipients.h"
+#include "IMessageAttachment.h"
+#include "IMessageBridge.h"
+#include "IMessageBus.h"
+#include "IMessageContext.h"
+#include "IMessageHandler.h"
+#include "IMessageInterceptor.h"
+#include "IMessageReceiver.h"
+#include "IMessageSender.h"
+#include "IMessageSubscription.h"
+#include "IMessageTracer.h"
+#include "IMessageTracerBreakpoint.h"
+#include "IMessageTransport.h"
+#include "IMutableMessageContext.h"
diff --git a/Engine/Source/Runtime/Messaging/Public/Deprecated/Messaging.h b/Engine/Source/Runtime/Messaging/Public/Deprecated/Messaging.h
--- a/Engine/Source/Runtime/Messaging/Public/Deprecated/Messaging.h
+++ b/Engine/Source/Runtime/Messaging/Public/Deprecated/Messaging.h
@@ -3,21 +3,14 @@
 #pragma once
 
 
-/* Public Dependencies
- *****************************************************************************/
-
 #include "CoreUObject.h"
 
-
-/* Public Includes
- *****************************************************************************/
-
 #include "IMessageAttachment.h"
 #include "IMessageContext.h"
 #include "IMessageInterceptor.h"
 #include "IMessageHandler.h"
-#include "IReceiveMessages.h"
-#include "ISendMessages.h"
+#include "IMessageReceiver.h"
+#include "IMessageSender.h"
 #include "IMessageSubscription.h"
 #include "IAuthorizeMessageRecipients.h"
 #include "IMessageTracerBreakpoint.h"
@@ -28,10 +21,6 @@
 #include "IMessageBridge.h"
 #include "IMessagingModule.h"
 
-
-/* Common helpers
- *****************************************************************************/
-
 #include "FileMessageAttachment.h"
 #include "MessageBridgeBuilder.h"
 #include "MessageEndpoint.h"
diff --git a/Engine/Source/Runtime/Messaging/Public/Helpers/MessageBridgeBuilder.h b/Engine/Source/Runtime/Messaging/Public/Helpers/MessageBridgeBuilder.h
--- a/Engine/Source/Runtime/Messaging/Public/Helpers/MessageBridgeBuilder.h
+++ b/Engine/Source/Runtime/Messaging/Public/Helpers/MessageBridgeBuilder.h
@@ -29,7 +29,7 @@ public:
 	 *
 	 * @param InBus The message bus to attach the bridge to.
 	 */
-	FMessageBridgeBuilder( const IMessageBusRef& Bus )
+	FMessageBridgeBuilder(const TSharedRef<IMessageBus, ESPMode::ThreadSafe>& Bus)
 		: Address(FMessageAddress::NewAddress())
 		, BusPtr(Bus)
 		, Disabled(false)
@@ -56,7 +56,7 @@ public:
 	 * @param InTransport The transport technology to use.
 	 * @return This instance (for method chaining).
 	 */
-	FMessageBridgeBuilder& UsingTransport( const IMessageTransportRef& InTransport )
+	FMessageBridgeBuilder& UsingTransport(const TSharedRef<IMessageTransport, ESPMode::ThreadSafe>& InTransport)
 	{
 		Transport = InTransport;
 
@@ -71,7 +71,7 @@ public:
 	 * @param InAddress The address to set.
 	 * @return This instance (for method chaining).
 	 */
-	FMessageBridgeBuilder& WithAddress( const FMessageAddress& InAddress )
+	FMessageBridgeBuilder& WithAddress(const FMessageAddress& InAddress)
 	{
 		Address = InAddress;
 
@@ -85,13 +85,13 @@ public:
 	 *
 	 * @return A new message bridge, or nullptr if it couldn't be built.
 	 */
-	IMessageBridgePtr Build()
+	TSharedPtr<IMessageBridge, ESPMode::ThreadSafe> Build()
 	{
-		IMessageBridgePtr Bridge;
+		TSharedPtr<IMessageBridge, ESPMode::ThreadSafe> Bridge;
 
 		check(Transport.IsValid());
 
-		IMessageBusPtr Bus = BusPtr.Pin();
+		auto Bus = BusPtr.Pin();
 
 		if (Bus.IsValid())
 		{
@@ -118,7 +118,7 @@ public:
 	 *
 	 * @return A new message bridge, or nullptr if it couldn't be built.
 	 */
-	operator IMessageBridgePtr()
+	operator TSharedPtr<IMessageBridge, ESPMode::ThreadSafe>()
 	{
 		return Build();
 	}
@@ -129,11 +129,11 @@ private:
 	FMessageAddress Address;
 
 	/** Holds a weak pointer to the message bus to attach to. */
-	IMessageBusWeakPtr BusPtr;
+	TWeakPtr<IMessageBus, ESPMode::ThreadSafe> BusPtr;
 
 	/** Holds a flag indicating whether the bridge should be disabled. */
 	bool Disabled;
 
 	/** Holds a reference to the message transport technology. */
-	IMessageTransportPtr Transport;
+	TSharedPtr<IMessageTransport, ESPMode::ThreadSafe> Transport;
 };
diff --git a/Engine/Source/Runtime/Messaging/Public/Helpers/MessageEndpoint.h b/Engine/Source/Runtime/Messaging/Public/Helpers/MessageEndpoint.h
--- a/Engine/Source/Runtime/Messaging/Public/Helpers/MessageEndpoint.h
+++ b/Engine/Source/Runtime/Messaging/Public/Helpers/MessageEndpoint.h
@@ -2,13 +2,14 @@
 
 #pragma once
 
-#include "IMessageAttachment.h"
 #include "IMessageBus.h"
 #include "IMessageContext.h"
 #include "IMessageHandler.h"
-#include "IReceiveMessages.h"
-#include "ISendMessages.h"
-#include "TaskGraphInterfaces.h"
+#include "IMessageReceiver.h"
+#include "IMessageSender.h"
+
+
+class IMessageHandler;
 
 
 /**
@@ -48,8 +49,8 @@ DECLARE_DELEGATE_RetVal_OneParam(bool, FOnMessageEndpointReceiveMessage, const I
  */
 class FMessageEndpoint
 	: public TSharedFromThis<FMessageEndpoint, ESPMode::ThreadSafe>
-	, public IReceiveMessages
-	, public ISendMessages
+	, public IMessageReceiver
+	, public IMessageSender
 {
 public:
 
@@ -70,7 +71,7 @@ public:
 	 * @param InBus The message bus to attach this endpoint to.
 	 * @param InHandlers The collection of message handlers to register.
 	 */
-	FMessageEndpoint( const FName& InName, const IMessageBusRef& InBus, const TArray<IMessageHandlerPtr>& InHandlers )
+	FMessageEndpoint(const FName& InName, const TSharedRef<IMessageBus, ESPMode::ThreadSafe>& InBus, const TArray<TSharedPtr<IMessageHandler, ESPMode::ThreadSafe>>& InHandlers)
 		: Address(FMessageAddress::NewAddress())
 		, BusPtr(InBus)
 		, Enabled(true)
@@ -85,7 +86,7 @@ public:
 	/** Destructor. */
 	~FMessageEndpoint()
 	{
-		IMessageBusPtr Bus = BusPtr.Pin();
+		auto Bus = BusPtr.Pin();
 
 		if (Bus.IsValid())
 		{
@@ -167,7 +168,7 @@ public:
 	 *
 	 * @param NamedThread The name of the thread to receive messages on.
 	 */
-	void SetRecipientThread( const ENamedThreads::Type& NamedThread )
+	void SetRecipientThread(const ENamedThreads::Type& NamedThread)
 	{
 		RecipientThread = ENamedThreads::GetThreadIndex(NamedThread);
 	}
@@ -183,9 +184,9 @@ public:
 	 * @param Context The context of the message to defer.
 	 * @param Delay The time delay.
 	 */
-	void Defer( const IMessageContextRef& Context, const FTimespan& Delay )
+	void Defer(const IMessageContextRef& Context, const FTimespan& Delay)
 	{
-		IMessageBusPtr Bus = GetBusIfEnabled();
+		TSharedPtr<IMessageBus, ESPMode::ThreadSafe> Bus = GetBusIfEnabled();
 
 		if (Bus.IsValid())
 		{
@@ -202,9 +203,9 @@ public:
 	 * @param Recipients The list of message recipients to forward the message to.
 	 * @param Delay The time delay.
 	 */
-	void Forward( const IMessageContextRef& Context, const TArray<FMessageAddress>& Recipients, const FTimespan& Delay )
+	void Forward(const IMessageContextRef& Context, const TArray<FMessageAddress>& Recipients, const FTimespan& Delay)
 	{
-		IMessageBusPtr Bus = GetBusIfEnabled();
+		TSharedPtr<IMessageBus, ESPMode::ThreadSafe> Bus = GetBusIfEnabled();
 
 		if (Bus.IsValid())
 		{
@@ -222,9 +223,9 @@ public:
 	 * @param Delay The delay after which to publish the message.
 	 * @param Expiration The time at which the message expires.
 	 */
-	void Publish( void* Message, UScriptStruct* TypeInfo, EMessageScope Scope, const FTimespan& Delay, const FDateTime& Expiration )
+	void Publish(void* Message, UScriptStruct* TypeInfo, EMessageScope Scope, const FTimespan& Delay, const FDateTime& Expiration)
 	{
-		IMessageBusPtr Bus = GetBusIfEnabled();
+		TSharedPtr<IMessageBus, ESPMode::ThreadSafe> Bus = GetBusIfEnabled();
 
 		if (Bus.IsValid())
 		{
@@ -242,9 +243,9 @@ public:
 	 * @param Delay The delay after which to send the message.
 	 * @param Expiration The time at which the message expires.
 	 */
-	void Send( void* Message, UScriptStruct* TypeInfo, const IMessageAttachmentPtr& Attachment, const TArray<FMessageAddress>& Recipients, const FTimespan& Delay, const FDateTime& Expiration )
+	void Send(void* Message, UScriptStruct* TypeInfo, const TSharedPtr<IMessageAttachment, ESPMode::ThreadSafe>& Attachment, const TArray<FMessageAddress>& Recipients, const FTimespan& Delay, const FDateTime& Expiration)
 	{
-		IMessageBusPtr Bus = GetBusIfEnabled();
+		TSharedPtr<IMessageBus, ESPMode::ThreadSafe> Bus = GetBusIfEnabled();
 
 		if (Bus.IsValid())
 		{
@@ -258,9 +259,9 @@ public:
 	 * @param MessageType The type name of the messages to subscribe to.
 	 * @param ScopeRange The range of message scopes to include in the subscription.
 	 */
-	void Subscribe( const FName& MessageType, const FMessageScopeRange& ScopeRange )
+	void Subscribe(const FName& MessageType, const FMessageScopeRange& ScopeRange)
 	{
-		IMessageBusPtr Bus = GetBusIfEnabled();
+		TSharedPtr<IMessageBus, ESPMode::ThreadSafe> Bus = GetBusIfEnabled();
 
 		if (Bus.IsValid())
 		{
@@ -274,9 +275,9 @@ public:
 	 * @param MessageType The type of message to unsubscribe (NAME_All = all types).
 	 * @see Subscribe
 	 */
-	void Unsubscribe( const FName& TopicPattern )
+	void Unsubscribe(const FName& TopicPattern)
 	{
-		IMessageBusPtr Bus = GetBusIfEnabled();
+		TSharedPtr<IMessageBus, ESPMode::ThreadSafe> Bus = GetBusIfEnabled();
 
 		if (Bus.IsValid())
 		{
@@ -363,7 +364,7 @@ public:
 	 * @return true if a message was received, false if the inbox was empty.
 	 * @see DisableInbox, EnableInbox, IsInboxEnabled, ProcessInbox
 	 */
-	bool ReceiveFromInbox( IMessageContextPtr& OutContext )
+	bool ReceiveFromInbox(IMessageContextPtr& OutContext)
 	{
 		return Inbox.Dequeue(OutContext);
 	}
@@ -407,7 +408,7 @@ public:
 		return true;
 	}
 
-	virtual void ReceiveMessage( const IMessageContextRef& Context ) override
+	virtual void ReceiveMessage(const IMessageContextRef& Context) override
 	{
 		if (!Enabled)
 		{
@@ -438,7 +439,7 @@ public:
 		return Address;
 	}
 
-	virtual void NotifyMessageError( const IMessageContextRef& Context, const FString& Error ) override
+	virtual void NotifyMessageError(const IMessageContextRef& Context, const FString& Error) override
 	{
 		ErrorDelegate.ExecuteIfBound(Context, Error);
 	}
@@ -453,7 +454,7 @@ public:
 	 * @param Context The context of the message to forward.
 	 * @param Recipient The address of the recipient to forward the message to.
 	 */
-	void Forward( const IMessageContextRef& Context, const FMessageAddress& Recipient )
+	void Forward(const IMessageContextRef& Context, const FMessageAddress& Recipient)
 	{
 		Forward(Context, TArrayBuilder<FMessageAddress>().Add(Recipient), FTimespan::Zero());
 	}
@@ -468,7 +469,7 @@ public:
 	 * @param ForwardingScope The scope of the forwarded message.
 	 * @param Delay The delay after which to publish the message.
 	 */
-	void Forward( const IMessageContextRef& Context, const FMessageAddress& Recipient, const FTimespan& Delay )
+	void Forward(const IMessageContextRef& Context, const FMessageAddress& Recipient, const FTimespan& Delay)
 	{
 		Forward(Context, TArrayBuilder<FMessageAddress>().Add(Recipient), Delay);
 	}
@@ -482,7 +483,7 @@ public:
 	 * @param Recipients The list of message recipients to forward the message to.
 	 * @param ForwardingScope The scope of the forwarded message.
 	 */
-	void Forward( const IMessageContextRef& Context, const TArray<FMessageAddress>& Recipients )
+	void Forward(const IMessageContextRef& Context, const TArray<FMessageAddress>& Recipients)
 	{
 		Forward(Context, Recipients, FTimespan::Zero());
 	}
@@ -493,7 +494,7 @@ public:
 	 * @param Message The message to publish.
 	 */
 	template<typename MessageType>
-	void Publish( MessageType* Message )
+	void Publish(MessageType* Message)
 	{
 		Publish(Message, MessageType::StaticStruct(), EMessageScope::Network, FTimespan::Zero(), FDateTime::MaxValue());
 	}
@@ -505,7 +506,7 @@ public:
 	 * @param Scope The message scope.
 	 */
 	template<typename MessageType>
-	void Publish( MessageType* Message, EMessageScope Scope )
+	void Publish(MessageType* Message, EMessageScope Scope)
 	{
 		Publish(Message, MessageType::StaticStruct(), Scope, FTimespan::Zero(), FDateTime::MaxValue());
 	}
@@ -517,7 +518,7 @@ public:
 	 * @param Delay The delay after which to publish the message.
 	 */
 	template<typename MessageType>
-	void Publish( MessageType* Message, const FTimespan& Delay )
+	void Publish(MessageType* Message, const FTimespan& Delay)
 	{
 		Publish(Message, MessageType::StaticStruct(), EMessageScope::Network, Delay, FDateTime::MaxValue());
 	}
@@ -530,7 +531,7 @@ public:
 	 * @param Delay The delay after which to publish the message.
 	 */
 	template<typename MessageType>
-	void Publish( MessageType* Message, EMessageScope Scope, const FTimespan& Delay )
+	void Publish(MessageType* Message, EMessageScope Scope, const FTimespan& Delay)
 	{
 		Publish(Message, MessageType::StaticStruct(), Scope, Delay, FDateTime::MaxValue());
 	}
@@ -545,7 +546,7 @@ public:
 	 * @param Expiration The time at which the message expires.
 	 */
 	template<typename MessageType>
-	void Publish( MessageType* Message, EMessageScope Scope, const FTimespan& Delay, const FDateTime& Expiration )
+	void Publish(MessageType* Message, EMessageScope Scope, const FTimespan& Delay, const FDateTime& Expiration)
 	{
 		Publish(Message, MessageType::StaticStruct(), Scope, Delay, Expiration);
 	}
@@ -558,7 +559,7 @@ public:
 	 * @param Recipient The message recipient.
 	 */
 	template<typename MessageType>
-	void Send( MessageType* Message, const FMessageAddress& Recipient )
+	void Send(MessageType* Message, const FMessageAddress& Recipient)
 	{
 		Send(Message, MessageType::StaticStruct(), nullptr, TArrayBuilder<FMessageAddress>().Add(Recipient), FTimespan::Zero(), FDateTime::MaxValue());
 	}
@@ -572,7 +573,7 @@ public:
 	 * @param Delay The delay after which to send the message.
 	 */
 	template<typename MessageType>
-	void Send( MessageType* Message, const FMessageAddress& Recipient, const FTimespan& Delay )
+	void Send(MessageType* Message, const FMessageAddress& Recipient, const FTimespan& Delay)
 	{
 		Send(Message, MessageType::StaticStruct(), nullptr, TArrayBuilder<FMessageAddress>().Add(Recipient), Delay, FDateTime::MaxValue());
 	}
@@ -587,7 +588,7 @@ public:
 	 * @param Delay The delay after which to send the message.
 	 */
 	template<typename MessageType>
-	void Send( MessageType* Message, const FMessageAddress& Recipient, const FTimespan& Delay, const FDateTime& Expiration )
+	void Send(MessageType* Message, const FMessageAddress& Recipient, const FTimespan& Delay, const FDateTime& Expiration)
 	{
 		Send(Message, MessageType::StaticStruct(), nullptr, TArrayBuilder<FMessageAddress>().Add(Recipient), Delay, Expiration);
 	}
@@ -601,7 +602,7 @@ public:
 	 * @param Recipient The message recipient.
 	 */
 	template<typename MessageType>
-	void Send( MessageType* Message, const IMessageAttachmentPtr& Attachment, const FMessageAddress& Recipient )
+	void Send(MessageType* Message, const TSharedPtr<IMessageAttachment, ESPMode::ThreadSafe>& Attachment, const FMessageAddress& Recipient)
 	{
 		Send(Message, MessageType::StaticStruct(), Attachment, TArrayBuilder<FMessageAddress>().Add(Recipient), FTimespan::Zero(), FDateTime::MaxValue());
 	}
@@ -617,7 +618,7 @@ public:
 	 * @param Delay The delay after which to send the message.
 	 */
 	template<typename MessageType>
-	void Send( MessageType* Message, const IMessageAttachmentPtr& Attachment, const FMessageAddress& Recipient, const FDateTime& Expiration, const FTimespan& Delay )
+	void Send(MessageType* Message, const TSharedPtr<IMessageAttachment, ESPMode::ThreadSafe>& Attachment, const FMessageAddress& Recipient, const FDateTime& Expiration, const FTimespan& Delay)
 	{
 		Send(Message, MessageType::StaticStruct(), Attachment, TArrayBuilder<FMessageAddress>().Add(Recipient), Delay, Expiration);
 	}
@@ -630,7 +631,7 @@ public:
 	 * @param Recipients The message recipients.
 	 */
 	template<typename MessageType>
-	void Send( MessageType* Message, const TArray<FMessageAddress>& Recipients )
+	void Send(MessageType* Message, const TArray<FMessageAddress>& Recipients)
 	{
 		Send(Message, MessageType::StaticStruct(), nullptr, Recipients, FTimespan::Zero(), FDateTime::MaxValue());
 	}
@@ -644,7 +645,7 @@ public:
 	 * @param Delay The delay after which to send the message.
 	 */
 	template<typename MessageType>
-	void Send( MessageType* Message, const TArray<FMessageAddress>& Recipients, const FTimespan& Delay )
+	void Send(MessageType* Message, const TArray<FMessageAddress>& Recipients, const FTimespan& Delay)
 	{
 		Send(Message, MessageType::StaticStruct(), nullptr, Recipients, Delay, FDateTime::MaxValue());
 	}
@@ -659,7 +660,7 @@ public:
 	 * @param Delay The delay after which to send the message.
 	 */
 	template<typename MessageType>
-	void Send( MessageType* Message, const IMessageAttachmentPtr& Attachment, const TArray<FMessageAddress>& Recipients, const FTimespan& Delay )
+	void Send(MessageType* Message, const TSharedPtr<IMessageAttachment, ESPMode::ThreadSafe>& Attachment, const TArray<FMessageAddress>& Recipients, const FTimespan& Delay)
 	{
 		Send(Message, MessageType::StaticStruct(), Attachment, Recipients, Delay, FDateTime::MaxValue());
 	}
@@ -675,7 +676,7 @@ public:
 	 * @param Expiration The time at which the message expires.
 	 */
 	template<typename MessageType>
-	void Send( MessageType* Message, const IMessageAttachmentPtr& Attachment, const TArray<FMessageAddress>& Recipients, const FTimespan& Delay, const FDateTime& Expiration )
+	void Send(MessageType* Message, const TSharedPtr<IMessageAttachment, ESPMode::ThreadSafe>& Attachment, const TArray<FMessageAddress>& Recipients, const FTimespan& Delay, const FDateTime& Expiration)
 	{
 		Send(Message, MessageType::StaticStruct(), Attachment, Recipients, Delay, Expiration);
 	}
@@ -706,7 +707,7 @@ public:
 	 * @param ScopeRange The range of message scopes to include in the subscription.
 	 */
 	template<class MessageType>
-	void Subscribe( const FMessageScopeRange& ScopeRange )
+	void Subscribe(const FMessageScopeRange& ScopeRange)
 	{
 		Subscribe(MessageType::StaticStruct()->GetFName(), ScopeRange);
 	}
@@ -749,7 +750,7 @@ public:
 	 *
 	 * @param Endpoint The message endpoint to release.
 	 */
-	static void SafeRelease( TSharedPtr<FMessageEndpoint, ESPMode::ThreadSafe>& Endpoint )
+	static void SafeRelease(TSharedPtr<FMessageEndpoint, ESPMode::ThreadSafe>& Endpoint)
 	{
 		TWeakPtr<FMessageEndpoint, ESPMode::ThreadSafe> EndpointPtr = Endpoint;
 		Endpoint.Reset();
@@ -763,7 +764,7 @@ protected:
 	 *
 	 * @return The message bus.
 	 */
-	FORCEINLINE IMessageBusPtr GetBusIfEnabled() const
+	FORCEINLINE TSharedPtr<IMessageBus, ESPMode::ThreadSafe> GetBusIfEnabled() const
 	{
 		if (Enabled)
 		{
@@ -778,7 +779,7 @@ protected:
 	 *
 	 * @param Context The context of the message to handle.
 	 */
-	void ProcessMessage( const IMessageContextRef& Context )
+	void ProcessMessage(const IMessageContextRef& Context)
 	{
 		if (!Context->IsValid())
 		{
@@ -797,13 +798,13 @@ private:
 	const FMessageAddress Address;
 
 	/** Holds a weak pointer to the message bus. */
-	IMessageBusWeakPtr BusPtr;
+	TWeakPtr<IMessageBus, ESPMode::ThreadSafe> BusPtr;
 
 	/** Hold a flag indicating whether this endpoint is active. */
 	bool Enabled;
 
 	/** Holds the registered message handlers. */
-	TArray<IMessageHandlerPtr> Handlers;
+	TArray<TSharedPtr<IMessageHandler, ESPMode::ThreadSafe>> Handlers;
 
 	/** Holds the endpoint's unique identifier (for debugging purposes). */
 	const FGuid Id;
diff --git a/Engine/Source/Runtime/Messaging/Public/Helpers/MessageEndpointBuilder.h b/Engine/Source/Runtime/Messaging/Public/Helpers/MessageEndpointBuilder.h
--- a/Engine/Source/Runtime/Messaging/Public/Helpers/MessageEndpointBuilder.h
+++ b/Engine/Source/Runtime/Messaging/Public/Helpers/MessageEndpointBuilder.h
@@ -3,7 +3,6 @@
 #pragma once
 
 #include "IMessageBus.h"
-#include "IMessageHandler.h"
 #include "IMessagingModule.h"
 #include "MessageEndpoint.h"
 #include "MessageHandlers.h"
@@ -25,7 +24,7 @@ public:
 	 * @param InName The endpoint's name (for debugging purposes).
 	 * @param InBus The message bus to attach the endpoint to.
 	 */
-	FMessageEndpointBuilder( const FName& InName )
+	FMessageEndpointBuilder(const FName& InName)
 		: BusPtr(IMessagingModule::Get().GetDefaultBus())
 		, Disabled(false)
 		, InboxEnabled(false)
@@ -39,7 +38,7 @@ public:
 	 * @param InName The endpoint's name (for debugging purposes).
 	 * @param InBus The message bus to attach the endpoint to.
 	 */
-	FMessageEndpointBuilder( const FName& InName, const IMessageBusRef& InBus )
+	FMessageEndpointBuilder(const FName& InName, const TSharedRef<IMessageBus, ESPMode::ThreadSafe>& InBus)
 		: BusPtr(InBus)
 		, Disabled(false)
 		, InboxEnabled(false)
@@ -65,7 +64,7 @@ public:
 	 * @see WithCatchall, WithHandler
 	 */
 	template<typename MessageType, typename HandlerType>
-	FMessageEndpointBuilder& Handling( HandlerType* Handler, typename TRawMessageHandler<MessageType, HandlerType>::FuncType HandlerFunc )
+	FMessageEndpointBuilder& Handling(HandlerType* Handler, typename TRawMessageHandler<MessageType, HandlerType>::FuncType HandlerFunc)
 	{
 		Handlers.Add(MakeShareable(new TRawMessageHandler<MessageType, HandlerType>(Handler, MoveTemp(HandlerFunc))));
 
@@ -87,7 +86,7 @@ public:
 	 * @see WithCatchall, WithHandler
 	 */
 	template<typename MessageType>
-	FMessageEndpointBuilder& Handling( typename TFunctionMessageHandler<MessageType>::FuncType HandlerFunc )
+	FMessageEndpointBuilder& Handling(typename TFunctionMessageHandler<MessageType>::FuncType HandlerFunc)
 	{
 		Handlers.Add(MakeShareable(new TFunctionMessageHandler<MessageType>(MoveTemp(HandlerFunc))));
 
@@ -128,7 +127,7 @@ public:
 	 * @return This instance (for method chaining).
 	 * @see ReceivingOnAnyThread
 	 */
-	FMessageEndpointBuilder& ReceivingOnThread( ENamedThreads::Type NamedThread )
+	FMessageEndpointBuilder& ReceivingOnThread(ENamedThreads::Type NamedThread)
 	{
 		RecipientThread = NamedThread;
 
@@ -163,7 +162,7 @@ public:
 	 * @see WithHandler
 	 */
 	template<typename HandlerType>
-	FMessageEndpointBuilder& WithCatchall( HandlerType* Handler, typename TRawMessageCatchall<HandlerType>::FuncType HandlerFunc )
+	FMessageEndpointBuilder& WithCatchall(HandlerType* Handler, typename TRawMessageCatchall<HandlerType>::FuncType HandlerFunc)
 	{
 		Handlers.Add(MakeShareable(new TRawMessageCatchall<HandlerType>(Handler, MoveTemp(HandlerFunc))));
 
@@ -184,7 +183,7 @@ public:
 	 * @return This instance (for method chaining).
 	 * @see WithHandler
 	 */
-	FMessageEndpointBuilder& WithCatchall( FFunctionMessageCatchall::FuncType HandlerFunc )
+	FMessageEndpointBuilder& WithCatchall(FFunctionMessageCatchall::FuncType HandlerFunc)
 	{
 		Handlers.Add(MakeShareable(new FFunctionMessageCatchall(MoveTemp(HandlerFunc))));
 
@@ -201,7 +200,7 @@ public:
 	 * @return This instance (for method chaining).
 	 * @see Handling, WithCatchall
 	 */
-	FMessageEndpointBuilder& WithHandler( const IMessageHandlerRef& Handler )
+	FMessageEndpointBuilder& WithHandler(const TSharedRef<IMessageHandler, ESPMode::ThreadSafe>& Handler)
 	{
 		Handlers.Add(Handler);
 
@@ -232,8 +231,7 @@ public:
 	FMessageEndpointPtr Build()
 	{
 		FMessageEndpointPtr Endpoint;
-
-		IMessageBusPtr Bus = BusPtr.Pin();
+		TSharedPtr<IMessageBus, ESPMode::ThreadSafe> Bus = BusPtr.Pin();
 		
 		if (Bus.IsValid())
 		{
@@ -272,13 +270,13 @@ public:
 private:
 
 	/** Holds a reference to the message bus to attach to. */
-	IMessageBusWeakPtr BusPtr;
+	TWeakPtr<IMessageBus, ESPMode::ThreadSafe> BusPtr;
 
 	/** Holds a flag indicating whether the endpoint should be disabled. */
 	bool Disabled;
 
 	/** Holds the collection of message handlers to register. */
-	TArray<IMessageHandlerPtr> Handlers;
+	TArray<TSharedPtr<IMessageHandler, ESPMode::ThreadSafe>> Handlers;
 
 	/** Holds a flag indicating whether the inbox should be enabled. */
 	bool InboxEnabled;
diff --git a/Engine/Source/Runtime/Messaging/Public/Helpers/MessageHandlers.h b/Engine/Source/Runtime/Messaging/Public/Helpers/MessageHandlers.h
--- a/Engine/Source/Runtime/Messaging/Public/Helpers/MessageHandlers.h
+++ b/Engine/Source/Runtime/Messaging/Public/Helpers/MessageHandlers.h
@@ -4,6 +4,7 @@
 
 #include "CoreUObject.h"
 #include "IMessageContext.h"
+#include "IMessageHandler.h"
 
 
 /**
@@ -18,7 +19,7 @@ class TRawMessageCatchall
 public:
 
 	/** Type definition for function pointers that are compatible with this TRawMessageCatchall. */
-	typedef void (HandlerType::*FuncType)( const TSharedRef<IMessageContext, ESPMode::ThreadSafe>& );
+	typedef void (HandlerType::*FuncType)(const TSharedRef<IMessageContext, ESPMode::ThreadSafe>&);
 
 public:
 
@@ -28,7 +29,7 @@ public:
 	 * @param InHandler The object that will handle the messages.
 	 * @param InFunc The object's message handling function.
 	 */
-	TRawMessageCatchall( HandlerType* InHandler, FuncType InFunc )
+	TRawMessageCatchall(HandlerType* InHandler, FuncType InFunc)
 		: Handler(InHandler)
 		, Func(InFunc)
 	{
@@ -40,9 +41,9 @@ public:
 
 public:
 
-	// IMessageHandler interface
+	//~ IMessageHandler interface
 	
-	virtual void HandleMessage( const TSharedRef<IMessageContext, ESPMode::ThreadSafe>& Context ) override
+	virtual void HandleMessage(const TSharedRef<IMessageContext, ESPMode::ThreadSafe>& Context) override
 	{
 		(Handler->*Func)(Context);
 	}
@@ -70,7 +71,7 @@ class TRawMessageHandler
 public:
 
 	/** Type definition for function pointers that are compatible with this TRawMessageHandler. */
-	typedef void (HandlerType::*FuncType)( const MessageType&, const TSharedRef<IMessageContext, ESPMode::ThreadSafe>& );
+	typedef void (HandlerType::*FuncType)(const MessageType&, const TSharedRef<IMessageContext, ESPMode::ThreadSafe>&);
 
 public:
 
@@ -80,7 +81,7 @@ public:
 	 * @param InHandler The object that will handle the messages.
 	 * @param InFunc The object's message handling function.
 	 */
-	TRawMessageHandler( HandlerType* InHandler, FuncType InFunc )
+	TRawMessageHandler(HandlerType* InHandler, FuncType InFunc)
 		: Handler(InHandler)
 		, Func(InFunc)
 	{
@@ -92,9 +93,9 @@ public:
 
 public:
 
-	// IMessageHandler interface
+	//~ IMessageHandler interface
 	
-	virtual void HandleMessage( const TSharedRef<IMessageContext, ESPMode::ThreadSafe>& Context ) override
+	virtual void HandleMessage(const TSharedRef<IMessageContext, ESPMode::ThreadSafe>& Context) override
 	{
 		if (Context->GetMessageType() == MessageType::StaticStruct()->GetFName())
 		{
@@ -130,7 +131,7 @@ public:
 	 *
 	 * @param InFunc The object's message handling function.
 	 */
-	FFunctionMessageCatchall( FuncType InFunc )
+	FFunctionMessageCatchall(FuncType InFunc)
 		: Func(MoveTemp(InFunc))
 	{ }
 
@@ -139,9 +140,9 @@ public:
 
 public:
 
-	// IMessageHandler interface
+	//~ IMessageHandler interface
 	
-	virtual void HandleMessage( const TSharedRef<IMessageContext, ESPMode::ThreadSafe>& Context ) override
+	virtual void HandleMessage(const TSharedRef<IMessageContext, ESPMode::ThreadSafe>& Context) override
 	{
 		Func(Context);
 	}
@@ -174,7 +175,7 @@ public:
 	 *
 	 * @param InHandlerFunc The object's message handling function.
 	 */
-	TFunctionMessageHandler( FuncType InFunc )
+	TFunctionMessageHandler(FuncType InFunc)
 		: Func(MoveTemp(InFunc))
 	{ }
 
@@ -183,9 +184,9 @@ public:
 
 public:
 
-	// IMessageHandler interface
+	//~ IMessageHandler interface
 	
-	virtual void HandleMessage( const TSharedRef<IMessageContext, ESPMode::ThreadSafe>& Context ) override
+	virtual void HandleMessage(const TSharedRef<IMessageContext, ESPMode::ThreadSafe>& Context) override
 	{
 		if (Context->GetMessageType() == MessageType::StaticStruct()->GetFName())
 		{
diff --git a/Engine/Source/Runtime/Messaging/Public/IAuthorizeMessageRecipients.h b/Engine/Source/Runtime/Messaging/Public/IAuthorizeMessageRecipients.h
--- a/Engine/Source/Runtime/Messaging/Public/IAuthorizeMessageRecipients.h
+++ b/Engine/Source/Runtime/Messaging/Public/IAuthorizeMessageRecipients.h
@@ -5,7 +5,7 @@
 
 struct FMessageAddress;
 class IMessageInterceptor;
-class IReceiveMessages;
+class IMessageReceiver;
 
 
 /**
@@ -22,7 +22,7 @@ public:
 	 * @param MessageType The type of messages to intercept.
 	 * @return true if the request was authorized, false otherwise.
 	 */
-	virtual bool AuthorizeInterceptor( const TSharedRef<IMessageInterceptor, ESPMode::ThreadSafe>& Interceptor, const FName& MessageType ) = 0;
+	virtual bool AuthorizeInterceptor(const TSharedRef<IMessageInterceptor, ESPMode::ThreadSafe>& Interceptor, const FName& MessageType) = 0;
 
 	/**
 	 * Authorizes a request to register the specified recipient.
@@ -31,7 +31,7 @@ public:
 	 * @param Address The recipient's address.
 	 * @return true if the request was authorized, false otherwise.
 	 */
-	virtual bool AuthorizeRegistration( const TSharedRef<IReceiveMessages, ESPMode::ThreadSafe>& Recipient, const FMessageAddress& Address ) = 0;
+	virtual bool AuthorizeRegistration(const TSharedRef<IMessageReceiver, ESPMode::ThreadSafe>& Recipient, const FMessageAddress& Address) = 0;
 
 	/**
 	 * Authorizes a request to add a subscription for the specified topic pattern.
@@ -40,7 +40,7 @@ public:
 	 * @param TopicPattern The message topic pattern to subscribe to.
 	 * @return true if the request is authorized, false otherwise.
 	 */
-	virtual bool AuthorizeSubscription( const TSharedRef<IReceiveMessages, ESPMode::ThreadSafe>& Subscriber, const FName& TopicPattern ) = 0;
+	virtual bool AuthorizeSubscription(const TSharedRef<IMessageReceiver, ESPMode::ThreadSafe>& Subscriber, const FName& TopicPattern) = 0;
 
 	/**
 	 * Authorizes a request to unregister the specified recipient.
@@ -48,7 +48,7 @@ public:
 	 * @param Address The address of the recipient to unregister.
 	 * @return true if the request was authorized, false otherwise.
 	 */
-	virtual bool AuthorizeUnregistration( const FMessageAddress& Address ) = 0;
+	virtual bool AuthorizeUnregistration(const FMessageAddress& Address) = 0;
 
 	/**
 	 * Authorizes a request to remove a subscription for the specified topic pattern.
@@ -57,17 +57,10 @@ public:
 	 * @param TopicPattern The message topic pattern to unsubscribe from.
 	 * @return true if the request is authorized, false otherwise.
 	 */
-	virtual bool AuthorizeUnsubscription( const TSharedRef<IReceiveMessages, ESPMode::ThreadSafe>& Subscriber, const FName& TopicPattern ) = 0;
+	virtual bool AuthorizeUnsubscription(const TSharedRef<IMessageReceiver, ESPMode::ThreadSafe>& Subscriber, const FName& TopicPattern) = 0;
 
 public:
 
 	/** Virtual destructor. */
 	virtual ~IAuthorizeMessageRecipients() { }
 };
-
-
-/** Type definition for shared pointers to instances of IAuthorizeMessageRecipients. */
-typedef TSharedPtr<IAuthorizeMessageRecipients> IAuthorizeMessageRecipientsPtr;
-
-/** Type definition for shared references to instances of IAuthorizeMessageRecipients. */
-typedef TSharedRef<IAuthorizeMessageRecipients> IAuthorizeMessageRecipientsRef;
diff --git a/Engine/Source/Runtime/Messaging/Public/IMessageAttachment.h b/Engine/Source/Runtime/Messaging/Public/IMessageAttachment.h
--- a/Engine/Source/Runtime/Messaging/Public/IMessageAttachment.h
+++ b/Engine/Source/Runtime/Messaging/Public/IMessageAttachment.h
@@ -36,10 +36,3 @@ public:
 	/** Virtual destructor. */
 	virtual ~IMessageAttachment() { }
 };
-
-
-/** Type definition for shared pointers to instances of IMessageAttachment. */
-typedef TSharedPtr<IMessageAttachment, ESPMode::ThreadSafe> IMessageAttachmentPtr;
-
-/** Type definition for shared references to instances of IMessageAttachment. */
-typedef TSharedRef<IMessageAttachment, ESPMode::ThreadSafe> IMessageAttachmentRef;
diff --git a/Engine/Source/Runtime/Messaging/Public/IMessageBridge.h b/Engine/Source/Runtime/Messaging/Public/IMessageBridge.h
--- a/Engine/Source/Runtime/Messaging/Public/IMessageBridge.h
+++ b/Engine/Source/Runtime/Messaging/Public/IMessageBridge.h
@@ -54,10 +54,3 @@ public:
 	/** Virtual destructor. */
 	virtual ~IMessageBridge() { }
 };
-
-
-/** Type definition for shared pointers to instances of IMessageBridge. */
-typedef TSharedPtr<IMessageBridge, ESPMode::ThreadSafe> IMessageBridgePtr;
-
-/** Type definition for shared references to instances of IMessageBridge. */
-typedef TSharedRef<IMessageBridge, ESPMode::ThreadSafe> IMessageBridgeRef;
diff --git a/Engine/Source/Runtime/Messaging/Public/IMessageBus.h b/Engine/Source/Runtime/Messaging/Public/IMessageBus.h
--- a/Engine/Source/Runtime/Messaging/Public/IMessageBus.h
+++ b/Engine/Source/Runtime/Messaging/Public/IMessageBus.h
@@ -10,8 +10,8 @@ class IMessageContext;
 class IMessageInterceptor;
 class IMessageSubscription;
 class IMessageTracer;
-class IReceiveMessages;
-class ISendMessages;
+class IMessageReceiver;
+class IMessageSender;
 class UScriptStruct;
 
 
@@ -113,7 +113,7 @@ public:
 	 * @param Forwarder The sender that forwards the message.
 	 * @see Publish, Send
 	 */
-	virtual void Forward( const TSharedRef<IMessageContext, ESPMode::ThreadSafe>& Context, const TArray<FMessageAddress>& Recipients, const FTimespan& Delay, const TSharedRef<ISendMessages, ESPMode::ThreadSafe>& Forwarder ) = 0;
+	virtual void Forward(const TSharedRef<IMessageContext, ESPMode::ThreadSafe>& Context, const TArray<FMessageAddress>& Recipients, const FTimespan& Delay, const TSharedRef<IMessageSender, ESPMode::ThreadSafe>& Forwarder) = 0;
 
 	/**
 	 * Gets the message bus tracer.
@@ -129,7 +129,7 @@ public:
 	 * @param MessageType The type of messages to intercept.
 	 * @see Unintercept
 	 */
-	virtual void Intercept( const TSharedRef<IMessageInterceptor, ESPMode::ThreadSafe>& Interceptor, const FName& MessageType ) = 0;
+	virtual void Intercept(const TSharedRef<IMessageInterceptor, ESPMode::ThreadSafe>& Interceptor, const FName& MessageType) = 0;
 
 	/**
 	 * Sends a message to subscribed recipients.
@@ -145,7 +145,7 @@ public:
 	 * @param Publisher The message publisher.
 	 * @see Forward, Send
 	 */
-	virtual void Publish( void* Message, UScriptStruct* TypeInfo, EMessageScope Scope, const FTimespan& Delay, const FDateTime& Expiration, const TSharedRef<ISendMessages, ESPMode::ThreadSafe>& Publisher ) = 0;
+	virtual void Publish(void* Message, UScriptStruct* TypeInfo, EMessageScope Scope, const FTimespan& Delay, const FDateTime& Expiration, const TSharedRef<IMessageSender, ESPMode::ThreadSafe>& Publisher) = 0;
 
 	/**
 	 * Registers a message recipient with the message bus.
@@ -154,7 +154,7 @@ public:
 	 * @param Recipient The message recipient.
 	 * @see Unregister
 	 */
-	virtual void Register( const FMessageAddress& Address, const TSharedRef<IReceiveMessages, ESPMode::ThreadSafe>& Recipient) = 0;
+	virtual void Register(const FMessageAddress& Address, const TSharedRef<IMessageReceiver, ESPMode::ThreadSafe>& Recipient) = 0;
 
 	/**
 	 * Sends a message to multiple recipients.
@@ -171,7 +171,7 @@ public:
 	 * @param Sender The message sender.
 	 * @see Forward, Publish
 	 */
-	virtual void Send( void* Message, UScriptStruct* TypeInfo, const TSharedPtr<IMessageAttachment, ESPMode::ThreadSafe>& Attachment, const TArray<FMessageAddress>& Recipients, const FTimespan& Delay, const FDateTime& Expiration, const TSharedRef<ISendMessages, ESPMode::ThreadSafe>& Sender ) = 0;
+	virtual void Send(void* Message, UScriptStruct* TypeInfo, const TSharedPtr<IMessageAttachment, ESPMode::ThreadSafe>& Attachment, const TArray<FMessageAddress>& Recipients, const FTimespan& Delay, const FDateTime& Expiration, const TSharedRef<IMessageSender, ESPMode::ThreadSafe>& Sender) = 0;
 
 	/**
 	 * Shuts down the message bus.
@@ -192,7 +192,7 @@ public:
 	 * @return The added subscription, or nullptr if the subscription failed.
 	 * @see Unsubscribe
 	 */
-	virtual TSharedPtr<IMessageSubscription, ESPMode::ThreadSafe> Subscribe( const TSharedRef<IReceiveMessages, ESPMode::ThreadSafe>& Subscriber, const FName& MessageType, const TRange<EMessageScope>& ScopeRange ) = 0;
+	virtual TSharedPtr<IMessageSubscription, ESPMode::ThreadSafe> Subscribe(const TSharedRef<IMessageReceiver, ESPMode::ThreadSafe>& Subscriber, const FName& MessageType, const TRange<EMessageScope>& ScopeRange) = 0;
 
 	/**
 	 * Removes an interceptor for messages of the specified type.
@@ -201,7 +201,7 @@ public:
 	 * @param MessageType The type of messages to stop intercepting.
 	 * @see Intercept
 	 */
-	virtual void Unintercept( const TSharedRef<IMessageInterceptor, ESPMode::ThreadSafe>& Interceptor, const FName& MessageType ) = 0;
+	virtual void Unintercept(const TSharedRef<IMessageInterceptor, ESPMode::ThreadSafe>& Interceptor, const FName& MessageType) = 0;
 
 	/**
 	 * Unregisters a message recipient from the message bus.
@@ -209,7 +209,7 @@ public:
 	 * @param Address The address of the recipient to unregister.
 	 * @see Register
 	 */
-	virtual void Unregister( const FMessageAddress& Address ) = 0;
+	virtual void Unregister(const FMessageAddress& Address) = 0;
 
 	/**
 	 * Cancels the specified message subscription.
@@ -218,7 +218,7 @@ public:
 	 * @param MessageType The type of messages to unsubscribe from (NAME_All = all types).
 	 * @see Subscribe
 	 */
-	virtual void Unsubscribe( const TSharedRef<IReceiveMessages, ESPMode::ThreadSafe>& Subscriber, const FName& MessageType ) = 0;
+	virtual void Unsubscribe(const TSharedRef<IMessageReceiver, ESPMode::ThreadSafe>& Subscriber, const FName& MessageType) = 0;
 
 public:
 
@@ -235,13 +235,3 @@ public:
 	/** Virtual destructor. */
 	virtual ~IMessageBus() { }
 };
-
-
-/** Type definition for shared pointers to instances of IMessageBus. */
-typedef TSharedPtr<IMessageBus, ESPMode::ThreadSafe> IMessageBusPtr;
-
-/** Type definition for shared references to instances of IMessageBus. */
-typedef TSharedRef<IMessageBus, ESPMode::ThreadSafe> IMessageBusRef;
-
-/** Type definition for weak pointers to instances of IMessageBus. */
-typedef TWeakPtr<IMessageBus, ESPMode::ThreadSafe> IMessageBusWeakPtr;
diff --git a/Engine/Source/Runtime/Messaging/Public/IMessageContext.h b/Engine/Source/Runtime/Messaging/Public/IMessageContext.h
--- a/Engine/Source/Runtime/Messaging/Public/IMessageContext.h
+++ b/Engine/Source/Runtime/Messaging/Public/IMessageContext.h
@@ -2,6 +2,7 @@
 
 #pragma once
 
+#include "CoreUObject.h"
 #include "TaskGraphInterfaces.h"
 
 
@@ -28,7 +29,7 @@ public:
 	 * @param Y The second address to compare.
 	 * @return true if the addresses are equal, false otherwise.
 	 */
-	friend bool operator==( const FMessageAddress& X, const FMessageAddress& Y )
+	friend bool operator==(const FMessageAddress& X, const FMessageAddress& Y)
 	{
 		return (X.UniqueId == Y.UniqueId);
 	}
@@ -40,7 +41,7 @@ public:
 	 * @param Y The second address to compare.
 	 * @return true if the addresses are not equal, false otherwise.
 	 */
-	friend bool operator!=( const FMessageAddress& X, const FMessageAddress& Y )
+	friend bool operator!=(const FMessageAddress& X, const FMessageAddress& Y)
 	{
 		return (X.UniqueId != Y.UniqueId);
 	}
@@ -51,7 +52,7 @@ public:
 	 * @param Ar The archive to serialize from or into.
 	 * @param G The address to serialize.
 	 */
-	friend FArchive& operator<<( FArchive& Ar, FMessageAddress& A )
+	friend FArchive& operator<<(FArchive& Ar, FMessageAddress& A)
 	{
 		return Ar << A.UniqueId;
 	}
@@ -98,7 +99,7 @@ public:
 	 * @param Address The address to calculate the hash for.
 	 * @return The hash.
 	 */
-	friend uint32 GetTypeHash( const FMessageAddress& Address )
+	friend uint32 GetTypeHash(const FMessageAddress& Address)
 	{
 		return FCrc::MemCrc_DEPRECATED(&Address.UniqueId, sizeof(FGuid));
 	}
@@ -227,7 +228,7 @@ public:
 	 * Gets the message data.
 	 *
 	 * @return A pointer to the message data.
-	 * @see GetMessageType, GetMessageTypeInfo
+	 * @see GetAttachment, GetMessageType, GetMessageTypeInfo
 	 */
 	virtual const void* GetMessage() const = 0;
 
@@ -250,6 +251,7 @@ public:
 	 * Gets the list of message recipients.
 	 *
 	 * @return Message recipients.
+	 * @see GetSender
 	 */
 	virtual const TArray<FMessageAddress>& GetRecipients() const = 0;
 
@@ -264,6 +266,7 @@ public:
 	 * Gets the sender's address.
 	 *
 	 * @return Sender address.
+	 * @see GetRecipients, GetSenderThread
 	 */
 	virtual const FMessageAddress& GetSender() const = 0;
 
@@ -271,6 +274,7 @@ public:
 	 * Gets the name of the thread from which the message was sent.
 	 *
 	 * @return Sender threat name.
+	 * @see GetSender
 	 */
 	virtual ENamedThreads::Type GetSenderThread() const = 0;
 
@@ -278,6 +282,7 @@ public:
 	 * Gets the time at which the message was forwarded.
 	 *
 	 * @return Time forwarded.
+	 * @see GetTimeSent, IsForwarded
 	 */
 	virtual const FDateTime& GetTimeForwarded() const = 0;
 
@@ -285,6 +290,7 @@ public:
 	 * Gets the time at which the message was sent.
 	 *
 	 * @return Time sent.
+	 * @see GetTimeForwarded
 	 */
 	virtual const FDateTime& GetTimeSent() const = 0;
 
@@ -294,6 +300,7 @@ public:
 	 * Gets the name of the message type.
 	 *
 	 * @return Message type name.
+	 * @see GetMessage, GetMessageTypeInfo
 	 */
 	FName GetMessageType() const
 	{
@@ -309,7 +316,7 @@ public:
 	 * Checks whether this is a forwarded message.
 	 *
 	 * @return true if the message was forwarded, false otherwise.
-	 * @see GetOriginalContext, IsValid
+	 * @see GetOriginalContext, GetTimeForwarded, IsValid
 	 */
 	bool IsForwarded() const
 	{
diff --git a/Engine/Source/Runtime/Messaging/Public/IMessageHandler.h b/Engine/Source/Runtime/Messaging/Public/IMessageHandler.h
--- a/Engine/Source/Runtime/Messaging/Public/IMessageHandler.h
+++ b/Engine/Source/Runtime/Messaging/Public/IMessageHandler.h
@@ -15,17 +15,10 @@ public:
 	 *
 	 * @param Context The context of the message to handle.
 	 */
-	virtual void HandleMessage( const TSharedRef<IMessageContext, ESPMode::ThreadSafe>& Context ) = 0;
+	virtual void HandleMessage(const TSharedRef<IMessageContext, ESPMode::ThreadSafe>& Context) = 0;
 
 public:
 
 	/** Virtual destructor. */
 	virtual ~IMessageHandler() { }
 };
-
-
-/** Type definition for shared pointers to instances of IMessageHandler. */
-typedef TSharedPtr<IMessageHandler, ESPMode::ThreadSafe> IMessageHandlerPtr;
-
-/** Type definition for shared references to instances of IMessageHandler. */
-typedef TSharedRef<IMessageHandler, ESPMode::ThreadSafe> IMessageHandlerRef;
diff --git a/Engine/Source/Runtime/Messaging/Public/IMessageInterceptor.h b/Engine/Source/Runtime/Messaging/Public/IMessageInterceptor.h
--- a/Engine/Source/Runtime/Messaging/Public/IMessageInterceptor.h
+++ b/Engine/Source/Runtime/Messaging/Public/IMessageInterceptor.h
@@ -42,10 +42,3 @@ public:
 	/** Virtual destructor. */
 	virtual ~IMessageInterceptor() { }
 };
-
-
-/** Type definition for shared pointers to instances of IInterceptMessages. */
-typedef TSharedPtr<IMessageInterceptor, ESPMode::ThreadSafe> IMessageInterceptorPtr;
-
-/** Type definition for shared references to instances of IInterceptMessages. */
-typedef TSharedRef<IMessageInterceptor, ESPMode::ThreadSafe> IMessageInterceptorRef;
diff --git a/Engine/Source/Runtime/Messaging/Public/IMessageSender.h b/Engine/Source/Runtime/Messaging/Public/IMessageSender.h
new file mode 100644
--- /dev/null
+++ b/Engine/Source/Runtime/Messaging/Public/IMessageSender.h
@@ -0,0 +1,52 @@
+// Copyright 1998-2016 Epic Games, Inc. All Rights Reserved.
+
+#pragma once
+
+
+class IMessageContext;
+
+
+/**
+ * Interface for message senders.
+ *
+ * Classes that implement this interface are able to send messages on a message bus. Each message sender must be uniquely
+ * identifiable with a message address that is returned by the ISendMessages.GetSenderAddress method. It is recommended
+ * that implementors of this interface generate a GUID for each instance, which must not change throughout the lifetime of
+ * the instance.
+ *
+ * The sending of messages is accomplished with the IMessageBus.Forward, IMessageBus.Publish and IMessageBus.Send methods.
+ * In case an error occurs during the sending of a message, the ISendMessages.NotifyMessageError method will be called.
+ *
+ * This interface provides a rather low-level mechanism for receiving messages. Instead of implementing it, Most users will
+ * want to use an instance of FMessageEndpoint, which provides a much more convenient way of sending and receiving messages.
+ *
+ * @see FMessageEndpoint, IMessageBus, IReceiveMessages
+ */
+class IMessageSender
+{
+public:
+
+	/**
+	 * Gets the sender's address.
+	 *
+	 * @return The message address.
+	 */
+	virtual FMessageAddress GetSenderAddress() = 0;
+
+	/**
+	 * Notifies the sender of errors.
+	 *
+	 * @param Context The context of the message that generated the error.
+	 * @param Error The error string.
+	 */
+	virtual void NotifyMessageError( const TSharedRef<IMessageContext, ESPMode::ThreadSafe>& Context, const FString& Error ) = 0;
+
+public:
+
+	/** Virtual destructor. */
+	virtual ~IMessageSender() { }
+};
+
+
+DEPRECATED(4.15, "ISendMessages has been renamed to IMessageSender")
+typedef IMessageSender ISendMessages;
diff --git a/Engine/Source/Runtime/Messaging/Public/IMessageSubscription.h b/Engine/Source/Runtime/Messaging/Public/IMessageSubscription.h
--- a/Engine/Source/Runtime/Messaging/Public/IMessageSubscription.h
+++ b/Engine/Source/Runtime/Messaging/Public/IMessageSubscription.h
@@ -4,7 +4,7 @@
 
 
 enum class EMessageScope : uint8;
-class IReceiveMessages;
+class IMessageReceiver;
 
 
 /**
@@ -58,7 +58,7 @@ public:
 	 * @return The subscriber.
 	 * @see GetMessageType, GetScopeRange
 	 */
-	virtual const TWeakPtr<IReceiveMessages, ESPMode::ThreadSafe>& GetSubscriber() = 0;
+	virtual const TWeakPtr<IMessageReceiver, ESPMode::ThreadSafe>& GetSubscriber() = 0;
 
 	/**
 	 * Checks whether the subscription is enabled.
@@ -73,10 +73,3 @@ public:
 	/** Virtual destructor. */
 	virtual ~IMessageSubscription() { }
 };
-
-
-/** Type definition for shared pointers to instances of IMessageSubscription. */
-typedef TSharedPtr<IMessageSubscription, ESPMode::ThreadSafe> IMessageSubscriptionPtr;
-
-/** Type definition for shared references to instances of IMessageSubscription. */
-typedef TSharedRef<IMessageSubscription, ESPMode::ThreadSafe> IMessageSubscriptionRef;
diff --git a/Engine/Source/Runtime/Messaging/Public/IMessageTracer.h b/Engine/Source/Runtime/Messaging/Public/IMessageTracer.h
--- a/Engine/Source/Runtime/Messaging/Public/IMessageTracer.h
+++ b/Engine/Source/Runtime/Messaging/Public/IMessageTracer.h
@@ -2,48 +2,17 @@
 
 #pragma once
 
+#include "IMessageContext.h"
 
-class IMessageContext;
 
+struct FMessageTracerEndpointInfo;
+struct FMessageTracerMessageInfo;
+struct FMessageTracerTypeInfo;
 
-/** Type definition for shared pointers to instances of FMessageTracerAddressInfo. */
-typedef TSharedPtr<struct FMessageTracerAddressInfo> FMessageTracerAddressInfoPtr;
 
-/** Type definition for shared references to instances of FMessageTracerAddressInfo. */
-typedef TSharedRef<struct FMessageTracerAddressInfo> FMessageTracerAddressInfoRef;
-
-/** Type definition for shared pointers to instances of FMessageTracerDispatchState. */
-typedef TSharedPtr<struct FMessageTracerDispatchState> FMessageTracerDispatchStatePtr;
-
-/** Type definition for shared references to instances of FMessageTracerDispatchState. */
-typedef TSharedRef<struct FMessageTracerDispatchState> FMessageTracerDispatchStateRef;
-
-/** Type definition for shared pointers to instances of FMessageTracerEndpointInfo. */
-typedef TSharedPtr<struct FMessageTracerEndpointInfo> FMessageTracerEndpointInfoPtr;
-
-/** Type definition for shared references to instances of FMessageTracerEndpointInfo. */
-typedef TSharedRef<struct FMessageTracerEndpointInfo> FMessageTracerEndpointInfoRef;
-
-/** Type definition for shared pointers to instances of FMessageTracerInterceptorInfo. */
-typedef TSharedPtr<struct FMessageTracerInterceptorInfo> FMessageTracerInterceptorInfoPtr;
-
-/** Type definition for shared references to instances of FMessageTracerInterceptorInfo. */
-typedef TSharedRef<struct FMessageTracerInterceptorInfo> FMessageTracerInterceptorInfoRef;
-
-/** Type definition for shared pointers to instances of FMessageTracerMessageInfo. */
-typedef TSharedPtr<struct FMessageTracerMessageInfo> FMessageTracerMessageInfoPtr;
-
-/** Type definition for shared references to instances of FMessageTracerMessageInfo. */
-typedef TSharedRef<struct FMessageTracerMessageInfo> FMessageTracerMessageInfoRef;
-
-/** Type definition for shared pointers to instances of FMessageTracerTypeInfo. */
-typedef TSharedPtr<struct FMessageTracerTypeInfo> FMessageTracerTypeInfoPtr;
-
-/** Type definition for shared references to instances of FMessageTracerTypeInfo. */
-typedef TSharedRef<struct FMessageTracerTypeInfo> FMessageTracerTypeInfoRef;
-
-
-/** Enumerates tracer breakpoint states. */
+/**
+ * Enumerates tracer breakpoint states.
+ */
 enum class EMessageTracerBreakpointState
 {
 	/** The breakpoint is disabled. */
@@ -107,7 +76,7 @@ struct FMessageTracerDispatchState
 	EMessageTracerDispatchTypes DispatchType;
 
 	/** Holds the endpoint to which the message was or is being dispatched. */
-	FMessageTracerEndpointInfoPtr EndpointInfo;
+	TSharedPtr<FMessageTracerEndpointInfo> EndpointInfo;
 
 	/** The thread on which the endpoint receives messages. */
 	ENamedThreads::Type RecipientThread;
@@ -142,19 +111,19 @@ struct FMessageTracerAddressInfo
 struct FMessageTracerEndpointInfo
 {
 	/** Holds the recipient's address information. */
-	TMap<FMessageAddress, FMessageTracerAddressInfoPtr> AddressInfos;
+	TMap<FMessageAddress, TSharedPtr<FMessageTracerAddressInfo>> AddressInfos;
 
 	/** Holds the recipient's human readable name. */
 	FName Name;
 
 	/** Holds the list of messages received by this recipient. */
-	TArray<FMessageTracerMessageInfoPtr> ReceivedMessages;
+	TArray<TSharedPtr<FMessageTracerMessageInfo>> ReceivedMessages;
 
 	/** Holds a flag indicating whether this is a remote recipient. */
 	bool Remote;
 
 	/** Holds the list of messages sent by this recipient. */
-	TArray<FMessageTracerMessageInfoPtr> SentMessages;
+	TArray<TSharedPtr<FMessageTracerMessageInfo>> SentMessages;
 };
 
 
@@ -167,7 +136,7 @@ struct FMessageTracerInterceptorInfo
 	FName Name;
 
 	/** Holds the list of messages intercepted by this interceptor. */
-	TArray<FMessageTracerMessageInfoPtr> InterceptedMessages;
+	TArray<TSharedPtr<FMessageTracerMessageInfo>> InterceptedMessages;
 
 	/** Holds the time at which this interceptor was registered. */
 	double TimeRegistered;
@@ -186,13 +155,13 @@ struct FMessageTracerMessageInfo
 	IMessageContextPtr Context;
 
 	/** Holds the message's dispatch states per endpoint. */
-	TMap<FMessageTracerEndpointInfoPtr, FMessageTracerDispatchStatePtr> DispatchStates;
+	TMap<TSharedPtr<FMessageTracerEndpointInfo>, TSharedPtr<FMessageTracerDispatchState>> DispatchStates;
 
 	/** Whether the message was intercepted. */
 	bool Intercepted;
 
 	/** Pointer to the sender's endpoint information. */
-	FMessageTracerEndpointInfoPtr SenderInfo;
+	TSharedPtr<FMessageTracerEndpointInfo> SenderInfo;
 
 	/** Holds the time at which the message was routed (0.0 = pending). */
 	double TimeRouted;
@@ -201,7 +170,7 @@ struct FMessageTracerMessageInfo
 	double TimeSent;
 
 	/** Pointer to the message's type information. */
-	FMessageTracerTypeInfoPtr TypeInfo;
+	TSharedPtr<FMessageTracerTypeInfo> TypeInfo;
 };
 
 
@@ -211,7 +180,7 @@ struct FMessageTracerMessageInfo
 struct FMessageTracerTypeInfo
 {
 	/** Holds the collection of messages of this type. */
-	TArray<FMessageTracerMessageInfoPtr> Messages;
+	TArray<TSharedPtr<FMessageTracerMessageInfo>> Messages;
 
 	/** Holds a name of the message type. */
 	FName TypeName;
@@ -292,14 +261,14 @@ public:
 	 * @param OutEndpoints Will contain the list of endpoints.
 	 * @return The number of endpoints returned.
 	 */
-	virtual int32 GetEndpoints(TArray<FMessageTracerEndpointInfoPtr>& OutEndpoints) const = 0;
+	virtual int32 GetEndpoints(TArray<TSharedPtr<FMessageTracerEndpointInfo>>& OutEndpoints) const = 0;
 
 	/**
 	 * Gets the collection of known messages.
 	 *
 	 * @return The messages.
 	 */
-	virtual int32 GetMessages(TArray<FMessageTracerMessageInfoPtr>& OutMessages) const = 0;
+	virtual int32 GetMessages(TArray<TSharedPtr<FMessageTracerMessageInfo>>& OutMessages) const = 0;
 
 	/**
 	 * Gets the list of known message types filtered by name.
@@ -308,7 +277,7 @@ public:
 	 * @param OutTypes Will contain the list of message types.
 	 * @return The number of message types returned.
 	 */
-	virtual int32 GetMessageTypes(TArray<FMessageTracerTypeInfoPtr>& OutTypes) const = 0;
+	virtual int32 GetMessageTypes(TArray<TSharedPtr<FMessageTracerTypeInfo>>& OutTypes) const = 0;
 
 	/**
 	 * Checks whether there are any messages in the history.
@@ -319,6 +288,9 @@ public:
 
 public:
 
+	typedef TSharedRef<FMessageTracerMessageInfo> FMessageTracerMessageInfoRef;
+	typedef TSharedRef<FMessageTracerTypeInfo> FMessageTracerTypeInfoRef;
+
 	/**
 	 * A delegate that is executed when the collection of known messages has changed.
 	 *
@@ -348,10 +320,3 @@ protected:
 	/** Hidden destructor. The life time of a message tracer is managed by the message bus. */
 	virtual ~IMessageTracer() { }
 };
-
-
-/** Type definition for shared pointers to instances of IMessageTracer. */
-typedef TSharedPtr<IMessageTracer, ESPMode::ThreadSafe> IMessageTracerPtr;
-
-/** Type definition for shared references to instances of IMessageTracer. */
-typedef TSharedRef<IMessageTracer, ESPMode::ThreadSafe> IMessageTracerRef;
diff --git a/Engine/Source/Runtime/Messaging/Public/IMessageTracerBreakpoint.h b/Engine/Source/Runtime/Messaging/Public/IMessageTracerBreakpoint.h
--- a/Engine/Source/Runtime/Messaging/Public/IMessageTracerBreakpoint.h
+++ b/Engine/Source/Runtime/Messaging/Public/IMessageTracerBreakpoint.h
@@ -28,17 +28,10 @@ public:
 	 * @param Context The context of the message to break on.
 	 * @return true if the tracer should break, false otherwise.
 	 */
-	virtual bool ShouldBreak( const TSharedRef<IMessageContext, ESPMode::ThreadSafe>& Context ) const = 0;
+	virtual bool ShouldBreak(const TSharedRef<IMessageContext, ESPMode::ThreadSafe>& Context) const = 0;
 
 protected:
 
 	/** Hidden destructor. */
 	virtual ~IMessageTracerBreakpoint() { }
 };
-
-
-/** Type definition for shared pointers to instances of IMessageTracerBreakpoint. */
-typedef TSharedPtr<IMessageTracerBreakpoint, ESPMode::ThreadSafe> IMessageTracerBreakpointPtr;
-
-/** Type definition for shared references to instances of IMessageTracerBreakpoint. */
-typedef TSharedRef<IMessageTracerBreakpoint, ESPMode::ThreadSafe> IMessageTracerBreakpointRef;
diff --git a/Engine/Source/Runtime/Messaging/Public/IMessageTransport.h b/Engine/Source/Runtime/Messaging/Public/IMessageTransport.h
--- a/Engine/Source/Runtime/Messaging/Public/IMessageTransport.h
+++ b/Engine/Source/Runtime/Messaging/Public/IMessageTransport.h
@@ -4,6 +4,7 @@
 
 
 class IMessageAttachment;
+class IMessageContext;
 class IMessageData;
 
 
@@ -46,7 +47,7 @@ public:
 	 * @param Recipients The transport nodes to send the message to.
 	 * @return true if the message is being transported, false otherwise.
 	 */
-	virtual bool TransportMessage(const IMessageContextRef& Context, const TArray<FGuid>& Recipients) = 0;
+	virtual bool TransportMessage(const TSharedRef<IMessageContext, ESPMode::ThreadSafe>& Context, const TArray<FGuid>& Recipients) = 0;
 
 public:
 
@@ -69,8 +70,5 @@ protected:
 };
 
 
-/** Type definition for shared pointers to instances of ITransportMessages. */
-typedef TSharedPtr<IMessageTransport, ESPMode::ThreadSafe> IMessageTransportPtr;
-
 /** Type definition for shared references to instances of ITransportMessages. */
 typedef TSharedRef<IMessageTransport, ESPMode::ThreadSafe> IMessageTransportRef;
diff --git a/Engine/Source/Runtime/Messaging/Public/IMutableMessageContext.h b/Engine/Source/Runtime/Messaging/Public/IMutableMessageContext.h
--- a/Engine/Source/Runtime/Messaging/Public/IMutableMessageContext.h
+++ b/Engine/Source/Runtime/Messaging/Public/IMutableMessageContext.h
@@ -5,6 +5,8 @@
 #include "IMessageContext.h"
 
 
+enum class EMessageScope : uint8;
+struct FMessageAddress;
 class IMessageAttachment;
 class UScriptStruct;
 
@@ -22,14 +24,14 @@ public:
 	 *
 	 * @param Recipient The address of the recipient to add.
 	 */
-	virtual void AddRecipient( const FMessageAddress& Recipient ) = 0;
+	virtual void AddRecipient(const FMessageAddress& Recipient) = 0;
 
 	/**
 	 * Sets the optional message attachment.
 	 *
 	 * @param InAttachment The attachment to set.
 	 */
-	virtual void SetAttachment( const TSharedPtr<IMessageAttachment, ESPMode::ThreadSafe>& InAttachment ) = 0;
+	virtual void SetAttachment(const TSharedPtr<IMessageAttachment, ESPMode::ThreadSafe>& InAttachment) = 0;
 
 	/**
 	 * Sets the message.
@@ -37,14 +39,14 @@ public:
 	 * @param InMessage The message to set.
 	 * @param InTypeInfo The message's type information.
 	 */
-	virtual void SetMessage( void* InMessage, UScriptStruct* InTypeInfo ) = 0;
+	virtual void SetMessage(void* InMessage, UScriptStruct* InTypeInfo) = 0;
 
 	/**
 	 * Sets the date and time at which the message expires.
 	 *
 	 * @param InExpiration Expiration date and time.
 	 */
-	virtual void SetExpiration( const FDateTime& InExpiration ) = 0;
+	virtual void SetExpiration(const FDateTime& InExpiration) = 0;
 
 	/**
 	 * Sets the value of the header with the specified key.
@@ -52,38 +54,31 @@ public:
 	 * @param Key The header key.
 	 * @param Value The header value.
 	 */
-	virtual void SetHeader( const FName& Key, const FString& Value ) = 0;
+	virtual void SetHeader(const FName& Key, const FString& Value) = 0;
 
 	/**
 	 * Sets the message scope.
 	 *
 	 * @param InScope The message scope.
 	 */
-	virtual void SetScope( EMessageScope InScope ) = 0;
+	virtual void SetScope(EMessageScope InScope) = 0;
 
 	/**
 	 * Sets the address of the message's sender.
 	 *
 	 * @param InSender The message sender's address.
 	 */
-	virtual void SetSender( const FMessageAddress& InSender ) = 0;
+	virtual void SetSender(const FMessageAddress& InSender) = 0;
 
 	/**
 	 * Sets the date and time at which the message was sent.
 	 *
 	 * @param InTimeSent Send date and time.
 	 */
-	virtual void SetTimeSent( const FDateTime& InTimeSent ) = 0;
+	virtual void SetTimeSent(const FDateTime& InTimeSent) = 0;
 
 public:
 
 	/** Virtual destructor. */
 	virtual ~IMutableMessageContext() { }
 };
-
-
-/** Type definition for shared pointers to instances of IMutableMessageContext. */
-typedef TSharedPtr<IMutableMessageContext, ESPMode::ThreadSafe> IMutableMessageContextPtr;
-
-/** Type definition for shared references to instances of IMutableMessageContext. */
-typedef TSharedRef<IMutableMessageContext, ESPMode::ThreadSafe> IMutableMessageContextRef;
