Add a config setting for enabling delayed messaging, disabled by default

diff --git a/Engine/Config/BaseEngine.ini b/Engine/Config/BaseEngine.ini
--- a/Engine/Config/BaseEngine.ini
+++ b/Engine/Config/BaseEngine.ini
@@ -2024,3 +2024,6 @@ LiveLinkPublishingPort=11111
 
 [/Script/Engine.RendererSettings]
 r.GPUCrashDebugging=True
+
+[Messaging]
+bAllowDelayedMessaging=false
\ No newline at end of file
diff --git a/Engine/Source/Runtime/Messaging/Private/Bus/MessageRouter.cpp b/Engine/Source/Runtime/Messaging/Private/Bus/MessageRouter.cpp
--- a/Engine/Source/Runtime/Messaging/Private/Bus/MessageRouter.cpp
+++ b/Engine/Source/Runtime/Messaging/Private/Bus/MessageRouter.cpp
@@ -6,7 +6,7 @@
 #include "IMessageSubscription.h"
 #include "IMessageReceiver.h"
 #include "IMessageInterceptor.h"
-
+#include "Misc/ConfigCacheIni.h"
 
 /* FMessageRouter structors
  *****************************************************************************/
@@ -15,9 +15,12 @@ FMessageRouter::FMessageRouter()
 	: DelayedMessagesSequence(0)
 	, Stopping(false)
 	, Tracer(MakeShareable(new FMessageTracer()))
+	, bAllowDelayedMessaging(false)
 {
 	ActiveSubscriptions.FindOrAdd(NAME_All);
 	WorkEvent = FPlatformProcess::GetSynchEventFromPool();
+
+	GConfig->GetBool(TEXT("Messaging"), TEXT("bAllowDelayedMessaging"), bAllowDelayedMessaging, GEngineIni);
 }
 
 
@@ -351,7 +354,7 @@ void FMessageRouter::HandleRouteMessage(TSharedRef<IMessageContext, ESPMode::Thr
 	}
 
 	// dispatch the message
-	if (Context->GetTimeSent() > CurrentTime)
+	if (bAllowDelayedMessaging && (Context->GetTimeSent() > CurrentTime))
 	{
 		DelayedMessages.HeapPush(FDelayedMessage(Context, ++DelayedMessagesSequence));
 	}
diff --git a/Engine/Source/Runtime/Messaging/Private/Bus/MessageRouter.h b/Engine/Source/Runtime/Messaging/Private/Bus/MessageRouter.h
--- a/Engine/Source/Runtime/Messaging/Private/Bus/MessageRouter.h
+++ b/Engine/Source/Runtime/Messaging/Private/Bus/MessageRouter.h
@@ -286,4 +286,7 @@ private:
 
 	/** Holds an event signaling that work is available. */
 	FEvent* WorkEvent;
+
+	/** Whether or not to allow delayed messaging */
+	bool bAllowDelayedMessaging;
 };
