Modified MessageBus Publish and Send interface to take Annotations by reference
Added a GetForwarder function to IMessageContext to get to the forwarder sender even if an original context exist as GetSender will return that context sender instead.
Merged from CL 7431537

diff --git a/Engine/Plugins/Messaging/TcpMessaging/Source/TcpMessaging/Private/Transport/TcpDeserializedMessage.cpp b/Engine/Plugins/Messaging/TcpMessaging/Source/TcpMessaging/Private/Transport/TcpDeserializedMessage.cpp
--- a/Engine/Plugins/Messaging/TcpMessaging/Source/TcpMessaging/Private/Transport/TcpDeserializedMessage.cpp
+++ b/Engine/Plugins/Messaging/TcpMessaging/Source/TcpMessaging/Private/Transport/TcpDeserializedMessage.cpp
@@ -194,6 +194,12 @@ const FMessageAddress& FTcpDeserializedMessage::GetSender() const
 }
 
 
+const FMessageAddress& FTcpDeserializedMessage::GetForwarder() const
+{
+	return Sender;
+}
+
+
 ENamedThreads::Type FTcpDeserializedMessage::GetSenderThread() const
 {
 	return ENamedThreads::AnyThread;
diff --git a/Engine/Plugins/Messaging/TcpMessaging/Source/TcpMessaging/Private/Transport/TcpDeserializedMessage.h b/Engine/Plugins/Messaging/TcpMessaging/Source/TcpMessaging/Private/Transport/TcpDeserializedMessage.h
--- a/Engine/Plugins/Messaging/TcpMessaging/Source/TcpMessaging/Private/Transport/TcpDeserializedMessage.h
+++ b/Engine/Plugins/Messaging/TcpMessaging/Source/TcpMessaging/Private/Transport/TcpDeserializedMessage.h
@@ -49,6 +49,7 @@ public:
 	virtual EMessageScope GetScope() const override;
 	virtual EMessageFlags GetFlags() const override;
 	virtual const FMessageAddress& GetSender() const override;
+	virtual const FMessageAddress& GetForwarder() const override;
 	virtual ENamedThreads::Type GetSenderThread() const override;
 	virtual const FDateTime& GetTimeForwarded() const override;
 	virtual const FDateTime& GetTimeSent() const override;
diff --git a/Engine/Plugins/Messaging/UdpMessaging/Source/UdpMessaging/Private/Tests/UdpMessagingTestTypes.h b/Engine/Plugins/Messaging/UdpMessaging/Source/UdpMessaging/Private/Tests/UdpMessagingTestTypes.h
--- a/Engine/Plugins/Messaging/UdpMessaging/Source/UdpMessaging/Private/Tests/UdpMessagingTestTypes.h
+++ b/Engine/Plugins/Messaging/UdpMessaging/Source/UdpMessaging/Private/Tests/UdpMessagingTestTypes.h
@@ -71,8 +71,8 @@ public:
 	virtual const TArray<FMessageAddress>& GetRecipients() const override { return Recipients; }
 	virtual EMessageScope GetScope() const override { return Scope; }
 	virtual EMessageFlags GetFlags() const override { return Flags; }
-
 	virtual const FMessageAddress& GetSender() const override { return Sender; }
+	virtual const FMessageAddress& GetForwarder() const override { return Sender; }
 	virtual ENamedThreads::Type GetSenderThread() const override { return SenderThread; }
 	virtual const FDateTime& GetTimeForwarded() const override { return TimeSent; }
 	virtual const FDateTime& GetTimeSent() const override { return TimeSent; }
diff --git a/Engine/Plugins/Messaging/UdpMessaging/Source/UdpMessaging/Private/Transport/UdpDeserializedMessage.cpp b/Engine/Plugins/Messaging/UdpMessaging/Source/UdpMessaging/Private/Transport/UdpDeserializedMessage.cpp
--- a/Engine/Plugins/Messaging/UdpMessaging/Source/UdpMessaging/Private/Transport/UdpDeserializedMessage.cpp
+++ b/Engine/Plugins/Messaging/UdpMessaging/Source/UdpMessaging/Private/Transport/UdpDeserializedMessage.cpp
@@ -117,6 +117,12 @@ const FMessageAddress& FUdpDeserializedMessage::GetSender() const
 }
 
 
+const FMessageAddress& FUdpDeserializedMessage::GetForwarder() const
+{
+	return Sender;
+}
+
+
 ENamedThreads::Type FUdpDeserializedMessage::GetSenderThread() const
 {
 	return ENamedThreads::AnyThread;
diff --git a/Engine/Plugins/Messaging/UdpMessaging/Source/UdpMessaging/Private/Transport/UdpDeserializedMessage.h b/Engine/Plugins/Messaging/UdpMessaging/Source/UdpMessaging/Private/Transport/UdpDeserializedMessage.h
--- a/Engine/Plugins/Messaging/UdpMessaging/Source/UdpMessaging/Private/Transport/UdpDeserializedMessage.h
+++ b/Engine/Plugins/Messaging/UdpMessaging/Source/UdpMessaging/Private/Transport/UdpDeserializedMessage.h
@@ -57,6 +57,7 @@ public:
 	virtual EMessageFlags GetFlags() const override;
 
 	virtual const FMessageAddress& GetSender() const override;
+	virtual const FMessageAddress& GetForwarder() const override;
 	virtual ENamedThreads::Type GetSenderThread() const override;
 	virtual const FDateTime& GetTimeForwarded() const override;
 	virtual const FDateTime& GetTimeSent() const override;
diff --git a/Engine/Source/Runtime/Messaging/Private/Bus/MessageBus.cpp b/Engine/Source/Runtime/Messaging/Private/Bus/MessageBus.cpp
--- a/Engine/Source/Runtime/Messaging/Private/Bus/MessageBus.cpp
+++ b/Engine/Source/Runtime/Messaging/Private/Bus/MessageBus.cpp
@@ -80,13 +80,13 @@ void FMessageBus::Publish(
 	void* Message,
 	UScriptStruct* TypeInfo,
 	EMessageScope Scope,
-	const TMap<FName, FString> Annotations,
+	const TMap<FName, FString>& Annotations,
 	const FTimespan& Delay,
 	const FDateTime& Expiration,
 	const TSharedRef<IMessageSender, ESPMode::ThreadSafe>& Publisher
 )
 {
-	Router->RouteMessage(MakeShareable(new FMessageContext(
+	Router->RouteMessage(MakeShared<FMessageContext, ESPMode::ThreadSafe>(
 		Message,
 		TypeInfo,
 		Annotations,
@@ -98,7 +98,7 @@ void FMessageBus::Publish(
 		FDateTime::UtcNow() + Delay,
 		Expiration,
 		FTaskGraphInterface::Get().GetCurrentThreadIfKnown()
-	)));
+	));
 }
 
 
@@ -112,7 +112,7 @@ void FMessageBus::Send(
 	void* Message,
 	UScriptStruct* TypeInfo,
 	EMessageFlags Flags,
-	const TMap<FName, FString> Annotations,
+	const TMap<FName, FString>& Annotations,
 	const TSharedPtr<IMessageAttachment, ESPMode::ThreadSafe>& Attachment,
 	const TArray<FMessageAddress>& Recipients,
 	const FTimespan& Delay,
@@ -120,7 +120,7 @@ void FMessageBus::Send(
 	const TSharedRef<IMessageSender, ESPMode::ThreadSafe>& Sender
 )
 {
-	Router->RouteMessage(MakeShareable(new FMessageContext(
+	Router->RouteMessage(MakeShared<FMessageContext, ESPMode::ThreadSafe>(
 		Message,
 		TypeInfo,
 		Annotations,
@@ -132,7 +132,7 @@ void FMessageBus::Send(
 		FDateTime::UtcNow() + Delay,
 		Expiration,
 		FTaskGraphInterface::Get().GetCurrentThreadIfKnown()
-	)));
+	));
 }
 
 
diff --git a/Engine/Source/Runtime/Messaging/Private/Bus/MessageBus.h b/Engine/Source/Runtime/Messaging/Private/Bus/MessageBus.h
--- a/Engine/Source/Runtime/Messaging/Private/Bus/MessageBus.h
+++ b/Engine/Source/Runtime/Messaging/Private/Bus/MessageBus.h
@@ -42,9 +42,9 @@ public:
 	virtual TSharedRef<IMessageTracer, ESPMode::ThreadSafe> GetTracer() override;
 	virtual void Intercept(const TSharedRef<IMessageInterceptor, ESPMode::ThreadSafe>& Interceptor, const FName& MessageType) override;
 	virtual FOnMessageBusShutdown& OnShutdown() override;
-	virtual void Publish(void* Message, UScriptStruct* TypeInfo, EMessageScope Scope, const TMap<FName, FString> Annotations, const FTimespan& Delay, const FDateTime& Expiration, const TSharedRef<IMessageSender, ESPMode::ThreadSafe>& Publisher) override;
+	virtual void Publish(void* Message, UScriptStruct* TypeInfo, EMessageScope Scope, const TMap<FName, FString>& Annotations, const FTimespan& Delay, const FDateTime& Expiration, const TSharedRef<IMessageSender, ESPMode::ThreadSafe>& Publisher) override;
 	virtual void Register(const FMessageAddress& Address, const TSharedRef<IMessageReceiver, ESPMode::ThreadSafe>& Recipient) override;
-	virtual void Send(void* Message, UScriptStruct* TypeInfo, EMessageFlags Flags, const TMap<FName, FString> Annotations, const TSharedPtr<IMessageAttachment, ESPMode::ThreadSafe>& Attachment, const TArray<FMessageAddress>& Recipients, const FTimespan& Delay, const FDateTime& Expiration, const TSharedRef<IMessageSender, ESPMode::ThreadSafe>& Sender) override;
+	virtual void Send(void* Message, UScriptStruct* TypeInfo, EMessageFlags Flags, const TMap<FName, FString>& Annotations, const TSharedPtr<IMessageAttachment, ESPMode::ThreadSafe>& Attachment, const TArray<FMessageAddress>& Recipients, const FTimespan& Delay, const FDateTime& Expiration, const TSharedRef<IMessageSender, ESPMode::ThreadSafe>& Sender) override;
 	virtual void Shutdown() override;
 	virtual TSharedPtr<IMessageSubscription, ESPMode::ThreadSafe> Subscribe(const TSharedRef<IMessageReceiver, ESPMode::ThreadSafe>& Subscriber, const FName& MessageType, const FMessageScopeRange& ScopeRange) override;
 	virtual void Unintercept(const TSharedRef<IMessageInterceptor, ESPMode::ThreadSafe>& Interceptor, const FName& MessageType) override;
diff --git a/Engine/Source/Runtime/Messaging/Private/Bus/MessageContext.cpp b/Engine/Source/Runtime/Messaging/Private/Bus/MessageContext.cpp
--- a/Engine/Source/Runtime/Messaging/Private/Bus/MessageContext.cpp
+++ b/Engine/Source/Runtime/Messaging/Private/Bus/MessageContext.cpp
@@ -115,6 +115,12 @@ const FMessageAddress& FMessageContext::GetSender() const
 }
 
 
+const FMessageAddress& FMessageContext::GetForwarder() const
+{
+	return Sender;
+}
+
+
 ENamedThreads::Type FMessageContext::GetSenderThread() const
 {
 	return SenderThread;
diff --git a/Engine/Source/Runtime/Messaging/Private/Bus/MessageContext.h b/Engine/Source/Runtime/Messaging/Private/Bus/MessageContext.h
--- a/Engine/Source/Runtime/Messaging/Private/Bus/MessageContext.h
+++ b/Engine/Source/Runtime/Messaging/Private/Bus/MessageContext.h
@@ -44,7 +44,7 @@ public:
 	FMessageContext(
 		void* InMessage,
 		UScriptStruct* InTypeInfo,
-		const TMap<FName, FString> InAnnotations,
+		const TMap<FName, FString>& InAnnotations,
 		const TSharedPtr<IMessageAttachment, ESPMode::ThreadSafe>& InAttachment,
 		const FMessageAddress& InSender,
 		const TArray<FMessageAddress>& InRecipients,
@@ -114,6 +114,7 @@ public:
 	virtual EMessageScope GetScope() const override;
 	virtual EMessageFlags GetFlags() const override;
 	virtual const FMessageAddress& GetSender() const override;
+	virtual const FMessageAddress& GetForwarder() const override;
 	virtual ENamedThreads::Type GetSenderThread() const override;
 	virtual const FDateTime& GetTimeForwarded() const override;
 	virtual const FDateTime& GetTimeSent() const override;
diff --git a/Engine/Source/Runtime/Messaging/Public/IMessageBus.h b/Engine/Source/Runtime/Messaging/Public/IMessageBus.h
--- a/Engine/Source/Runtime/Messaging/Public/IMessageBus.h
+++ b/Engine/Source/Runtime/Messaging/Public/IMessageBus.h
@@ -157,7 +157,7 @@ public:
 	 * @param Publisher The message publisher.
 	 * @see Forward, Send
 	 */
-	virtual void Publish(void* Message, UScriptStruct* TypeInfo, EMessageScope Scope, const TMap<FName, FString> Annotations, const FTimespan& Delay, const FDateTime& Expiration, const TSharedRef<IMessageSender, ESPMode::ThreadSafe>& Publisher) = 0;
+	virtual void Publish(void* Message, UScriptStruct* TypeInfo, EMessageScope Scope, const TMap<FName, FString>& Annotations, const FTimespan& Delay, const FDateTime& Expiration, const TSharedRef<IMessageSender, ESPMode::ThreadSafe>& Publisher) = 0;
 
 	/**
 	 * Registers a message recipient with the message bus.
@@ -185,7 +185,7 @@ public:
 	 * @param Sender The message sender.
 	 * @see Forward, Publish
 	 */
-	virtual void Send(void* Message, UScriptStruct* TypeInfo, EMessageFlags Flags, const TMap<FName, FString> Annotations, const TSharedPtr<IMessageAttachment, ESPMode::ThreadSafe>& Attachment, const TArray<FMessageAddress>& Recipients, const FTimespan& Delay, const FDateTime& Expiration, const TSharedRef<IMessageSender, ESPMode::ThreadSafe>& Sender) = 0;
+	virtual void Send(void* Message, UScriptStruct* TypeInfo, EMessageFlags Flags, const TMap<FName, FString>& Annotations, const TSharedPtr<IMessageAttachment, ESPMode::ThreadSafe>& Attachment, const TArray<FMessageAddress>& Recipients, const FTimespan& Delay, const FDateTime& Expiration, const TSharedRef<IMessageSender, ESPMode::ThreadSafe>& Sender) = 0;
 
 	/**
 	 * Shuts down the message bus.
diff --git a/Engine/Source/Runtime/Messaging/Public/IMessageContext.h b/Engine/Source/Runtime/Messaging/Public/IMessageContext.h
--- a/Engine/Source/Runtime/Messaging/Public/IMessageContext.h
+++ b/Engine/Source/Runtime/Messaging/Public/IMessageContext.h
@@ -295,6 +295,14 @@ public:
 	 */
 	virtual const FMessageAddress& GetSender() const = 0;
 
+	/**
+	 * Gets the forwarder's address.
+	 * @note Identical to GetSender if the message wasn't forwarded.
+	 * @return Forwarder address.
+	 * @see GetSender, GetRecipients, GetSenderThread
+	 */
+	virtual const FMessageAddress& GetForwarder() const = 0;
+	
 	/**
 	 * Gets the name of the thread from which the message was sent.
 	 *
