Fix crash when exiting in messaging module.

diff --git a/Engine/Source/Runtime/Messaging/Private/MessagingModule.cpp b/Engine/Source/Runtime/Messaging/Private/MessagingModule.cpp
--- a/Engine/Source/Runtime/Messaging/Private/MessagingModule.cpp
+++ b/Engine/Source/Runtime/Messaging/Private/MessagingModule.cpp
@@ -131,6 +131,10 @@ public:
 	virtual void ShutdownModule() override
 	{
 		ShutdownDefaultBus();
+
+#if PLATFORM_SUPPORTS_MESSAGEBUS
+		FCoreDelegates::OnPreExit.RemoveAll(this);
+#endif	//PLATFORM_SUPPORTS_MESSAGEBUS
 	}
 
 	virtual bool SupportsDynamicReloading() override
