Add TotalBytesLost, TotalBytesReceived, and TotalBytesInflight to FMessageTransportStatistics

diff --git a/Engine/Plugins/Messaging/UdpMessaging/Source/UdpMessaging/Private/Transport/UdpMessageProcessor.cpp b/Engine/Plugins/Messaging/UdpMessaging/Source/UdpMessaging/Private/Transport/UdpMessageProcessor.cpp
--- a/Engine/Plugins/Messaging/UdpMessaging/Source/UdpMessaging/Private/Transport/UdpMessageProcessor.cpp
+++ b/Engine/Plugins/Messaging/UdpMessaging/Source/UdpMessaging/Private/Transport/UdpMessageProcessor.cpp
@@ -666,6 +666,7 @@
 		}
 	}
 
+	NodeInfo.Statistics.TotalBytesReceived += DataChunk.Data.Num();
 	NodeInfo.Statistics.PacketsReceived++;
 	ReassembledMessage->Reassemble(DataChunk.SegmentNumber, DataChunk.SegmentOffset, DataChunk.Data, CurrentTime);
 
@@ -1014,6 +1015,7 @@
 	NodeInfo.Statistics.TotalBytesSent += SentInfo.BytesSent;
 	NodeInfo.Statistics.IPv4AsString = NodeInfo.Endpoint.ToString();
 	NodeInfo.Statistics.PacketsInFlight = NodeInfo.InflightSegments.Num();
+	NodeInfo.Statistics.BytesInflight = NodeInfo.InflightSegments.Num() * UDP_MESSAGING_SEGMENT_SIZE;
 	return MoveTemp(SentInfo);
 }
 
diff --git a/Engine/Plugins/Messaging/UdpMessaging/Source/UdpMessaging/Private/Transport/UdpMessageProcessor.h b/Engine/Plugins/Messaging/UdpMessaging/Source/UdpMessaging/Private/Transport/UdpMessageProcessor.h
--- a/Engine/Plugins/Messaging/UdpMessaging/Source/UdpMessaging/Private/Transport/UdpMessageProcessor.h
+++ b/Engine/Plugins/Messaging/UdpMessaging/Source/UdpMessaging/Private/Transport/UdpMessageProcessor.h
@@ -356,6 +356,8 @@
 				}
 				return bIsLost;
 			});
+
+			Statistics.TotalBytesLost += SegmentsLost * UDP_MESSAGING_SEGMENT_SIZE;
 			Statistics.PacketsLost += SegmentsLost;
 			return SegmentsLost;
 		}
diff --git a/Engine/Source/Runtime/Messaging/Public/INetworkMessagingExtension.h b/Engine/Source/Runtime/Messaging/Public/INetworkMessagingExtension.h
--- a/Engine/Source/Runtime/Messaging/Public/INetworkMessagingExtension.h
+++ b/Engine/Source/Runtime/Messaging/Public/INetworkMessagingExtension.h
@@ -14,6 +14,15 @@
 	/** Total number of bytes sent to the destination endpoint.*/
 	uint64 TotalBytesSent = 0;
 
+	/** Total number of bytes lost to the destination endpoint. */
+	uint64 TotalBytesLost = 0;
+
+	/** Total number of bytes received from endpoint. */
+	uint64 TotalBytesReceived = 0;
+	
+	/** Number of bytes currently in flight waiting for. */
+	uint64 BytesInflight = 0;
+
 	/** Total number of packets sent to the destination. */
 	uint64 PacketsSent = 0;
 
@@ -23,7 +32,7 @@
 	/** Number of packets acknowledged. */
 	uint64 PacketsAcked = 0;
 
-	/** Number of packets received by this endpoint. */
+	/** Number of packets received from this endpoint. */
 	uint64 PacketsReceived = 0;
 
 	/** Current packets in flight waiting for */
