Create a method to lookup the UUID for a FMessageAddress

diff --git a/Engine/Source/Runtime/Messaging/Private/Bridge/MessageBridge.cpp b/Engine/Source/Runtime/Messaging/Private/Bridge/MessageBridge.cpp
--- a/Engine/Source/Runtime/Messaging/Private/Bridge/MessageBridge.cpp
+++ b/Engine/Source/Runtime/Messaging/Private/Bridge/MessageBridge.cpp
@@ -103,7 +103,21 @@
 	return Enabled;
 }
 
+FGuid FMessageBridge::LookupAddress(const FMessageAddress &InAddress)
+{
+	if (InAddress.IsValid())
+	{
+		TArray<FGuid> RemoteNodes = AddressBook.GetNodesFor({InAddress});
+		ensure(RemoteNodes.Num() <= 1);
+		if (RemoteNodes.Num()>0 && RemoteNodes[0].IsValid())
+		{
+			return RemoteNodes[0];
+		}
+	}
+	return {};
+}
 
+
 /* IMessageReceiver interface
  *****************************************************************************/
 
diff --git a/Engine/Source/Runtime/Messaging/Private/Bridge/MessageBridge.h b/Engine/Source/Runtime/Messaging/Private/Bridge/MessageBridge.h
--- a/Engine/Source/Runtime/Messaging/Private/Bridge/MessageBridge.h
+++ b/Engine/Source/Runtime/Messaging/Private/Bridge/MessageBridge.h
@@ -63,6 +63,7 @@
 	virtual void Disable() override;
 	virtual void Enable() override;
 	virtual bool IsEnabled() const override;
+	virtual FGuid LookupAddress(const FMessageAddress& InAddress);
 
 public:
 
diff --git a/Engine/Source/Runtime/Messaging/Public/IMessageBridge.h b/Engine/Source/Runtime/Messaging/Public/IMessageBridge.h
--- a/Engine/Source/Runtime/Messaging/Public/IMessageBridge.h
+++ b/Engine/Source/Runtime/Messaging/Public/IMessageBridge.h
@@ -2,8 +2,9 @@
 
 #pragma once
 
+#include "IMessageContext.h"
 
-/**
+/*
  * Interface for message bridges.
  *
  * A message bridge connects a message bus with another messaging system. It is really just a regular
@@ -49,6 +50,14 @@
 	 */
 	virtual bool IsEnabled() const = 0;
 
+	/**
+	 * Provides the endpoint guid for the given address.
+	 *
+	 * @param MessageAddress to find remote endpoint guid.
+	 * @return FGuid which may be invalid.
+	 */
+	virtual FGuid LookupAddress(const FMessageAddress& InAddress) = 0;
+
 public:
 
 	/** Virtual destructor. */
