Deprecating ANY_PACKAGE.

diff --git a/Engine/Source/Runtime/HeadMountedDisplay/Private/XRMotionControllerBase.cpp b/Engine/Source/Runtime/HeadMountedDisplay/Private/XRMotionControllerBase.cpp
--- a/Engine/Source/Runtime/HeadMountedDisplay/Private/XRMotionControllerBase.cpp
+++ b/Engine/Source/Runtime/HeadMountedDisplay/Private/XRMotionControllerBase.cpp
@@ -3,8 +3,6 @@
 #include "XRMotionControllerBase.h"
 #include "UObject/UObjectGlobals.h" // for FindObject<>
 #include "UObject/Class.h" // for UEnum
-#include "UObject/ObjectMacros.h" // for ANY_PACKAGE
-#include "UObject/Package.h" // required for ANY_PACKAGE
 #include "InputCoreTypes.h" // for EControllerHand
 
 FName FXRMotionControllerBase::LeftHandSourceId(TEXT("Left"));
diff --git a/Engine/Source/Runtime/JsonUtilities/Private/JsonObjectConverter.cpp b/Engine/Source/Runtime/JsonUtilities/Private/JsonObjectConverter.cpp
--- a/Engine/Source/Runtime/JsonUtilities/Private/JsonObjectConverter.cpp
+++ b/Engine/Source/Runtime/JsonUtilities/Private/JsonObjectConverter.cpp
@@ -717,7 +717,7 @@ namespace
 				Obj->RemoveField(ObjectClassNameKey);
 				if (!ClassString.IsEmpty())
 				{
-					UClass* FoundClass = FindObject<UClass>(ANY_PACKAGE, *ClassString);
+					UClass* FoundClass = UClass::TryFindTypeSlow<UClass>(ClassString);
 					if (FoundClass)
 					{
 						PropertyClass = FoundClass;
diff --git a/Engine/Source/Runtime/Launch/Private/LaunchEngineLoop.cpp b/Engine/Source/Runtime/Launch/Private/LaunchEngineLoop.cpp
--- a/Engine/Source/Runtime/Launch/Private/LaunchEngineLoop.cpp
+++ b/Engine/Source/Runtime/Launch/Private/LaunchEngineLoop.cpp
@@ -2655,7 +2655,7 @@ PRAGMA_ENABLE_DEPRECATION_WARNINGS
 		bool bIsPossiblyCommandletName = Token.Len() && !Token.Contains(TEXT("-"));
 		if (bIsPossiblyCommandletName)
 		{
-			UClass* TempCommandletClass = FindObject<UClass>(ANY_PACKAGE, *(Token + TEXT("Commandlet")), false);
+			UClass* TempCommandletClass = FindFirstObject<UClass>(*(Token + TEXT("Commandlet")), EFindFirstObjectOptions::None, ELogVerbosity::Warning, TEXT("Looking for commandlet class"));
 			if (TempCommandletClass)
 			{
 				checkf(TempCommandletClass->IsChildOf(UCommandlet::StaticClass()), TEXT("It is not valid to have a class that ends with \"Commandlet\" that is not a UCommandlet subclass."));
@@ -3558,13 +3558,13 @@ int32 FEngineLoop::PreInitPostStartupScreen(const TCHAR* CmdLine)
 		{
 			checkf(PRIVATE_GIsRunningCommandlet, TEXT("This should have been set in PreInitPreStartupScreen"));
 
-			CommandletClass = FindObject<UClass>(ANY_PACKAGE,*Token,false);
+			CommandletClass = Cast<UClass>(StaticFindFirstObject(UClass::StaticClass(), *Token, EFindFirstObjectOptions::None, ELogVerbosity::Warning, TEXT("looking for commandlet")));
 			int32 PeriodIdx;
 			if (!CommandletClass && Token.FindChar('.', PeriodIdx))
 			{
 				// try to load module for commandlet specified before a period.
 				FModuleManager::Get().LoadModule(*Token.Left(PeriodIdx));
-				CommandletClass = FindObject<UClass>(ANY_PACKAGE, *Token, false);
+				CommandletClass = FindFirstObject<UClass>(*Token, EFindFirstObjectOptions::None, ELogVerbosity::Warning, TEXT("Looking for commandlet class"));
 			}
 			if (!CommandletClass)
 			{
@@ -3822,7 +3822,7 @@ int32 FEngineLoop::PreInitPostStartupScreen(const TCHAR* CmdLine)
 			if (bIsPossiblyUnrecognizedCommandlet)
 			{
 				// here we give people a reasonable warning if they tried to use the short name of a commandlet
-				UClass* TempCommandletClass = FindObject<UClass>(ANY_PACKAGE,*(Token+TEXT("Commandlet")),false);
+				UClass* TempCommandletClass = FindFirstObject<UClass>(*(Token + TEXT("Commandlet")), EFindFirstObjectOptions::None, ELogVerbosity::Warning, TEXT("Looking for commandlet class"));
 				if (TempCommandletClass)
 				{
 					UE_LOG(LogInit, Fatal, TEXT("You probably meant to call a commandlet. Please use the full name %s."), *(Token+TEXT("Commandlet")));
@@ -6404,4 +6404,4 @@ void FPreInitContext::Cleanup()
 	SlowTaskPtr = nullptr;
 }
 
-#undef LOCTEXT_NAMESPACE
+#undef LOCTEXT_NAMESPACE
\ No newline at end of file
diff --git a/Engine/Source/Runtime/LevelSequence/Private/LevelSequenceBindingReference.cpp b/Engine/Source/Runtime/LevelSequence/Private/LevelSequenceBindingReference.cpp
--- a/Engine/Source/Runtime/LevelSequence/Private/LevelSequenceBindingReference.cpp
+++ b/Engine/Source/Runtime/LevelSequence/Private/LevelSequenceBindingReference.cpp
@@ -135,7 +135,7 @@ UObject* ResolveByPath(UObject* InContext, const FString& InObjectPath)
 		}
 #endif
 
-		if (UObject* FoundObject = FindObject<UObject>(ANY_PACKAGE, *InObjectPath, false))
+		if (UObject* FoundObject = FindFirstObject<UObject>(*InObjectPath, EFindFirstObjectOptions::NativeFirst | EFindFirstObjectOptions::EnsureIfAmbiguous))
 		{
 			return FoundObject;
 		}
diff --git a/Engine/Source/Runtime/Messaging/Private/Bridge/MessageBridge.cpp b/Engine/Source/Runtime/Messaging/Private/Bridge/MessageBridge.cpp
--- a/Engine/Source/Runtime/Messaging/Private/Bridge/MessageBridge.cpp
+++ b/Engine/Source/Runtime/Messaging/Private/Bridge/MessageBridge.cpp
@@ -91,7 +91,7 @@
 	}
 	else
 	{
-		MessageSubscription = Bus->Subscribe(AsShared(), NAME_All, FMessageScopeRange::AtLeast(EMessageScope::Network));
+		MessageSubscription = Bus->Subscribe(AsShared(), IMessageBus::PATHNAME_All, FMessageScopeRange::AtLeast(EMessageScope::Network));
 	}
 
 	Enabled = true;
@@ -147,7 +147,7 @@
 
 void FMessageBridge::ReceiveMessage(const TSharedRef<IMessageContext, ESPMode::ThreadSafe>& Context)
 {
-	UE_LOG(LogMessaging, Verbose, TEXT("Received %s message from %s"), *Context->GetMessageType().ToString(), *Context->GetSender().ToString());
+	UE_LOG(LogMessaging, Verbose, TEXT("Received %s message from %s"), *Context->GetMessageTypePathName().ToString(), *Context->GetSender().ToString());
 
 	if (!Enabled)
 	{
@@ -220,7 +220,7 @@
 	{
 		FString RecipientStr = FString::JoinBy(Context->GetRecipients(), TEXT("+"), &FMessageAddress::ToString);
 		UE_LOG(LogMessaging, Verbose, TEXT("FMessageBridge::ReceiveTransportMessage: Received %s from %s for %s"),
-			*Context->GetMessageType().ToString(),
+			*Context->GetMessageTypePathName().ToString(),
 			*Context->GetSender().ToString(),
 			*RecipientStr);
 	}
@@ -240,7 +240,7 @@
 	// register newly discovered endpoints
 	if (!AddressBook.Contains(Context->GetSender()))
 	{
-		UE_LOG(LogMessaging, Verbose, TEXT("FMessageBridge::ReceiveTransportMessage: Registering new sender %s"), *Context->GetMessageType().ToString());
+		UE_LOG(LogMessaging, Verbose, TEXT("FMessageBridge::ReceiveTransportMessage: Registering new sender %s"), *Context->GetMessageTypePathName().ToString());
 
 		AddressBook.Add(Context->GetSender(), NodeId);
 		Bus->Register(Context->GetSender(), AsShared());
diff --git a/Engine/Source/Runtime/Messaging/Private/Bus/MessageBus.cpp b/Engine/Source/Runtime/Messaging/Private/Bus/MessageBus.cpp
--- a/Engine/Source/Runtime/Messaging/Private/Bus/MessageBus.cpp
+++ b/Engine/Source/Runtime/Messaging/Private/Bus/MessageBus.cpp
@@ -12,7 +12,7 @@
 
 
 
-
+const FTopLevelAssetPath IMessageBus::PATHNAME_All(TEXT("/Unknown"), NAME_All);
 
 /* FMessageBus structors
  *****************************************************************************/
@@ -51,7 +51,7 @@ void FMessageBus::Forward(
 		FString RecipientStr = FString::JoinBy(Context->GetRecipients(), TEXT("+"), &FMessageAddress::ToString);
 
 		UE_LOG(LogMessaging, Verbose, TEXT("Forwarding %s from %s to %s"),
-			*Context->GetMessageType().ToString(),
+			*Context->GetMessageTypePathName().ToString(),
 			*Context->GetSender().ToString(), *RecipientStr);
 	}
 
@@ -72,9 +72,9 @@ TSharedRef<IMessageTracer, ESPMode::ThreadSafe> FMessageBus::GetTracer()
 }
 
 
-void FMessageBus::Intercept(const TSharedRef<IMessageInterceptor, ESPMode::ThreadSafe>& Interceptor, const FName& MessageType)
+void FMessageBus::Intercept(const TSharedRef<IMessageInterceptor, ESPMode::ThreadSafe>& Interceptor, const FTopLevelAssetPath& MessageType)
 {
-	if (MessageType == NAME_None)
+	if (MessageType.IsNull())
 	{
 		return;
 	}
@@ -173,11 +173,11 @@ void FMessageBus::Shutdown()
 
 TSharedPtr<IMessageSubscription, ESPMode::ThreadSafe> FMessageBus::Subscribe(
 	const TSharedRef<IMessageReceiver, ESPMode::ThreadSafe>& Subscriber,
-	const FName& MessageType,
+	const FTopLevelAssetPath& MessageType,
 	const FMessageScopeRange& ScopeRange
 )
 {
-	if (MessageType != NAME_None)
+	if (!MessageType.IsNull())
 	{
 		if (!RecipientAuthorizer.IsValid() || RecipientAuthorizer->AuthorizeSubscription(Subscriber, MessageType))
 		{
@@ -193,9 +193,9 @@ TSharedPtr<IMessageSubscription, ESPMode::ThreadSafe> FMessageBus::Subscribe(
 }
 
 
-void FMessageBus::Unintercept(const TSharedRef<IMessageInterceptor, ESPMode::ThreadSafe>& Interceptor, const FName& MessageType)
+void FMessageBus::Unintercept(const TSharedRef<IMessageInterceptor, ESPMode::ThreadSafe>& Interceptor, const FTopLevelAssetPath& MessageType)
 {
-	if (MessageType != NAME_None)
+	if (!MessageType.IsNull())
 	{
 		UE_LOG(LogMessaging, Verbose, TEXT("Unintercepting %s"), *Interceptor->GetDebugName().ToString());
 		Router->RemoveInterceptor(Interceptor, MessageType);
@@ -214,9 +214,9 @@ void FMessageBus::Unregister(const FMessageAddress& Address)
 }
 
 
-void FMessageBus::Unsubscribe(const TSharedRef<IMessageReceiver, ESPMode::ThreadSafe>& Subscriber, const FName& MessageType)
+void FMessageBus::Unsubscribe(const TSharedRef<IMessageReceiver, ESPMode::ThreadSafe>& Subscriber, const FTopLevelAssetPath& MessageType)
 {
-	if (MessageType != NAME_None)
+	if (!MessageType.IsNull())
 	{
 		if (!RecipientAuthorizer.IsValid() || RecipientAuthorizer->AuthorizeUnsubscription(Subscriber, MessageType))
 		{
diff --git a/Engine/Source/Runtime/Messaging/Private/Bus/MessageBus.h b/Engine/Source/Runtime/Messaging/Private/Bus/MessageBus.h
--- a/Engine/Source/Runtime/Messaging/Private/Bus/MessageBus.h
+++ b/Engine/Source/Runtime/Messaging/Private/Bus/MessageBus.h
@@ -41,16 +41,16 @@ public:
 
 	virtual void Forward(const TSharedRef<IMessageContext, ESPMode::ThreadSafe>& Context, const TArray<FMessageAddress>& Recipients, const FTimespan& Delay, const TSharedRef<IMessageSender, ESPMode::ThreadSafe>& Forwarder) override;
 	virtual TSharedRef<IMessageTracer, ESPMode::ThreadSafe> GetTracer() override;
-	virtual void Intercept(const TSharedRef<IMessageInterceptor, ESPMode::ThreadSafe>& Interceptor, const FName& MessageType) override;
+	virtual void Intercept(const TSharedRef<IMessageInterceptor, ESPMode::ThreadSafe>& Interceptor, const FTopLevelAssetPath& MessageType) override;
 	virtual FOnMessageBusShutdown& OnShutdown() override;
 	virtual void Publish(void* Message, UScriptStruct* TypeInfo, EMessageScope Scope, const TMap<FName, FString>& Annotations, const FTimespan& Delay, const FDateTime& Expiration, const TSharedRef<IMessageSender, ESPMode::ThreadSafe>& Publisher) override;
 	virtual void Register(const FMessageAddress& Address, const TSharedRef<IMessageReceiver, ESPMode::ThreadSafe>& Recipient) override;
 	virtual void Send(void* Message, UScriptStruct* TypeInfo, EMessageFlags Flags, const TMap<FName, FString>& Annotations, const TSharedPtr<IMessageAttachment, ESPMode::ThreadSafe>& Attachment, const TArray<FMessageAddress>& Recipients, const FTimespan& Delay, const FDateTime& Expiration, const TSharedRef<IMessageSender, ESPMode::ThreadSafe>& Sender) override;
 	virtual void Shutdown() override;
-	virtual TSharedPtr<IMessageSubscription, ESPMode::ThreadSafe> Subscribe(const TSharedRef<IMessageReceiver, ESPMode::ThreadSafe>& Subscriber, const FName& MessageType, const FMessageScopeRange& ScopeRange) override;
-	virtual void Unintercept(const TSharedRef<IMessageInterceptor, ESPMode::ThreadSafe>& Interceptor, const FName& MessageType) override;
+	virtual TSharedPtr<IMessageSubscription, ESPMode::ThreadSafe> Subscribe(const TSharedRef<IMessageReceiver, ESPMode::ThreadSafe>& Subscriber, const FTopLevelAssetPath& MessageType, const FMessageScopeRange& ScopeRange) override;
+	virtual void Unintercept(const TSharedRef<IMessageInterceptor, ESPMode::ThreadSafe>& Interceptor, const FTopLevelAssetPath& MessageType) override;
 	virtual void Unregister(const FMessageAddress& Address) override;
-	virtual void Unsubscribe(const TSharedRef<IMessageReceiver, ESPMode::ThreadSafe>& Subscriber, const FName& MessageType) override;
+	virtual void Unsubscribe(const TSharedRef<IMessageReceiver, ESPMode::ThreadSafe>& Subscriber, const FTopLevelAssetPath& MessageType) override;
 
 	virtual void AddNotificationListener(const TSharedRef<IBusListener, ESPMode::ThreadSafe>& Listener) override;
 	virtual void RemoveNotificationListener(const TSharedRef<IBusListener, ESPMode::ThreadSafe>& Listener) override;
diff --git a/Engine/Source/Runtime/Messaging/Private/Bus/MessageRouter.cpp b/Engine/Source/Runtime/Messaging/Private/Bus/MessageRouter.cpp
--- a/Engine/Source/Runtime/Messaging/Private/Bus/MessageRouter.cpp
+++ b/Engine/Source/Runtime/Messaging/Private/Bus/MessageRouter.cpp
@@ -4,6 +4,7 @@
 #include "IMessagingModule.h"
 #include "HAL/PlatformProcess.h"
 #include "Bus/MessageDispatchTask.h"
+#include "IMessageBus.h"
 #include "IMessageSubscription.h"
 #include "IMessageReceiver.h"
 #include "IMessageInterceptor.h"
@@ -20,7 +21,7 @@ FMessageRouter::FMessageRouter()
 	, Tracer(MakeShared<FMessageTracer, ESPMode::ThreadSafe>())
 	, bAllowDelayedMessaging(false)
 {
-	ActiveSubscriptions.FindOrAdd(NAME_All);
+	ActiveSubscriptions.FindOrAdd(IMessageBus::PATHNAME_All);
 	WorkEvent = FPlatformProcess::GetSynchEventFromPool();
 
 	GConfig->GetBool(TEXT("Messaging"), TEXT("bAllowDelayedMessaging"), bAllowDelayedMessaging, GEngineIni);
@@ -128,7 +129,7 @@ void FMessageRouter::DispatchMessage(const TSharedRef<IMessageContext, ESPMode::
 			if (UE_GET_LOG_VERBOSITY(LogMessaging) >= ELogVerbosity::Verbose)
 			{
 				FString RecipientStr = FString::JoinBy(Context->GetRecipients(), TEXT("+"), &FMessageAddress::ToString);
-				UE_LOG(LogMessaging, Verbose, TEXT("Dispatching %s from %s to %s"), *Context->GetMessageType().ToString(), *Context->GetSender().ToString(), *RecipientStr);
+				UE_LOG(LogMessaging, Verbose, TEXT("Dispatching %s from %s to %s"), *Context->GetMessageTypePathName().ToString(), *Context->GetSender().ToString(), *RecipientStr);
 			}
 
 			FilterRecipients(Context, Recipients);
@@ -141,13 +142,13 @@ void FMessageRouter::DispatchMessage(const TSharedRef<IMessageContext, ESPMode::
 		// ... or from subscriptions
 		else
 		{
-			FilterSubscriptions(ActiveSubscriptions.FindOrAdd(Context->GetMessageType()), Context, Recipients);
-			FilterSubscriptions(ActiveSubscriptions.FindOrAdd(NAME_All), Context, Recipients);
+			FilterSubscriptions(ActiveSubscriptions.FindOrAdd(Context->GetMessageTypePathName()), Context, Recipients);
+			FilterSubscriptions(ActiveSubscriptions.FindOrAdd(IMessageBus::PATHNAME_All), Context, Recipients);
 
 			if (UE_GET_LOG_VERBOSITY(LogMessaging) >= ELogVerbosity::Verbose)
 			{
 				FString RecipientStr = FString::JoinBy(Context->GetRecipients(), TEXT("+"), &FMessageAddress::ToString);
-				UE_LOG(LogMessaging, Verbose, TEXT("Dispatching %s from %s to %s subscribers"), *Context->GetMessageType().ToString(), *Context->GetSender().ToString(), *RecipientStr);
+				UE_LOG(LogMessaging, Verbose, TEXT("Dispatching %s from %s to %s subscribers"), *Context->GetMessageTypePathName().ToString(), *Context->GetSender().ToString(), *RecipientStr);
 			}
 		}
 
@@ -278,7 +279,7 @@ void FMessageRouter::Tick()
 /* FMessageRouter callbacks
  *****************************************************************************/
 
-void FMessageRouter::HandleAddInterceptor(TSharedRef<IMessageInterceptor, ESPMode::ThreadSafe> Interceptor, FName MessageType)
+void FMessageRouter::HandleAddInterceptor(TSharedRef<IMessageInterceptor, ESPMode::ThreadSafe> Interceptor, FTopLevelAssetPath MessageType)
 {
 	UE_LOG(LogMessaging, Verbose, TEXT("Adding %s as intereceptor for %s messages"), *Interceptor->GetDebugName().ToString(), *MessageType.ToString());
 
@@ -308,19 +309,19 @@ void FMessageRouter::HandleAddSubscriber(TSharedRef<IMessageSubscription, ESPMod
 
 	if (Subscriber.IsValid())
 	{
-		UE_LOG(LogMessaging, Verbose, TEXT("Adding %s as a subscriber for %s messages"), *Subscriber->GetDebugName().ToString(), *Subscription->GetMessageType().ToString());
+		UE_LOG(LogMessaging, Verbose, TEXT("Adding %s as a subscriber for %s messages"), *Subscriber->GetDebugName().ToString(), *Subscription->GetMessageTypePathName().ToString());
 	}
 
-	ActiveSubscriptions.FindOrAdd(Subscription->GetMessageType()).AddUnique(Subscription);
+	ActiveSubscriptions.FindOrAdd(Subscription->GetMessageTypePathName()).AddUnique(Subscription);
 	Tracer->TraceAddedSubscription(Subscription);
 }
 
 
-void FMessageRouter::HandleRemoveInterceptor(TSharedRef<IMessageInterceptor, ESPMode::ThreadSafe> Interceptor, FName MessageType)
+void FMessageRouter::HandleRemoveInterceptor(TSharedRef<IMessageInterceptor, ESPMode::ThreadSafe> Interceptor, FTopLevelAssetPath MessageType)
 {
 	UE_LOG(LogMessaging, Verbose, TEXT("Removing %s as intereceptor for %s messages"), *Interceptor->GetDebugName().ToString(), *MessageType.ToString());
 
-	if (MessageType == NAME_All)
+	if (MessageType == IMessageBus::PATHNAME_All)
 	{
 		for (auto& InterceptorsPair : ActiveInterceptors)
 		{
@@ -351,7 +352,7 @@ void FMessageRouter::HandleRemoveRecipient(FMessageAddress Address)
 }
 
 
-void FMessageRouter::HandleRemoveSubscriber(TWeakPtr<IMessageReceiver, ESPMode::ThreadSafe> SubscriberPtr, FName MessageType)
+void FMessageRouter::HandleRemoveSubscriber(TWeakPtr<IMessageReceiver, ESPMode::ThreadSafe> SubscriberPtr, FTopLevelAssetPath MessageType)
 {
 	auto Subscriber = SubscriberPtr.Pin();
 
@@ -362,7 +363,7 @@ void FMessageRouter::HandleRemoveSubscriber(TWeakPtr<IMessageReceiver, ESPMode::
 
 	for (auto& SubscriptionsPair : ActiveSubscriptions)
 	{
-		if ((MessageType != NAME_All) && (MessageType != SubscriptionsPair.Key))
+		if ((MessageType != IMessageBus::PATHNAME_All) && (MessageType != SubscriptionsPair.Key))
 		{
 			continue;
 		}
@@ -375,7 +376,7 @@ void FMessageRouter::HandleRemoveSubscriber(TWeakPtr<IMessageReceiver, ESPMode::
 
 			if (Subscription->GetSubscriber().Pin() == Subscriber)
 			{
-				UE_LOG(LogMessaging, Verbose, TEXT("Removing %s as a subscriber for %s messages"), *Subscriber->GetDebugName().ToString(), *Subscription->GetMessageType().ToString());
+				UE_LOG(LogMessaging, Verbose, TEXT("Removing %s as a subscriber for %s messages"), *Subscriber->GetDebugName().ToString(), *Subscription->GetMessageTypePathName().ToString());
 
 				Subscriptions.RemoveAtSwap(SubscriptionIndex);
 				Tracer->TraceRemovedSubscription(Subscription.ToSharedRef(), MessageType);
@@ -389,12 +390,12 @@ void FMessageRouter::HandleRemoveSubscriber(TWeakPtr<IMessageReceiver, ESPMode::
 
 void FMessageRouter::HandleRouteMessage(TSharedRef<IMessageContext, ESPMode::ThreadSafe> Context)
 {
-	UE_LOG(LogMessaging, Verbose, TEXT("Routing %s message from %s"), *Context->GetMessageType().ToString(), *Context->GetSender().ToString());
+	UE_LOG(LogMessaging, Verbose, TEXT("Routing %s message from %s"), *Context->GetMessageTypePathName().ToString(), *Context->GetSender().ToString());
 
 	Tracer->TraceRoutedMessage(Context);
 
 	// intercept routing
-	auto& Interceptors = ActiveInterceptors.FindOrAdd(Context->GetMessageType());
+	auto& Interceptors = ActiveInterceptors.FindOrAdd(Context->GetMessageTypePathName());
 
 	for (auto& Interceptor : Interceptors)
 	{
diff --git a/Engine/Source/Runtime/Messaging/Private/Bus/MessageRouter.h b/Engine/Source/Runtime/Messaging/Private/Bus/MessageRouter.h
--- a/Engine/Source/Runtime/Messaging/Private/Bus/MessageRouter.h
+++ b/Engine/Source/Runtime/Messaging/Private/Bus/MessageRouter.h
@@ -44,7 +44,7 @@ public:
 	 * @param Interceptor The interceptor to add.
 	 * @param MessageType The type of messages to intercept.
 	 */
-	FORCEINLINE void AddInterceptor(const TSharedRef<IMessageInterceptor, ESPMode::ThreadSafe>& Interceptor, const FName& MessageType)
+	FORCEINLINE void AddInterceptor(const TSharedRef<IMessageInterceptor, ESPMode::ThreadSafe>& Interceptor, const FTopLevelAssetPath& MessageType)
 	{
 		EnqueueCommand(FSimpleDelegate::CreateRaw(this, &FMessageRouter::HandleAddInterceptor, Interceptor, MessageType));
 	}
@@ -86,7 +86,7 @@ public:
 	 * @param Interceptor The interceptor to remove.
 	 * @param MessageType The type of messages to stop intercepting.
 	 */
-	FORCEINLINE void RemoveInterceptor(const TSharedRef<IMessageInterceptor, ESPMode::ThreadSafe>& Interceptor, const FName& MessageType)
+	FORCEINLINE void RemoveInterceptor(const TSharedRef<IMessageInterceptor, ESPMode::ThreadSafe>& Interceptor, const FTopLevelAssetPath& MessageType)
 	{
 		EnqueueCommand(FSimpleDelegate::CreateRaw(this, &FMessageRouter::HandleRemoveInterceptor, Interceptor, MessageType));
 	}
@@ -107,7 +107,7 @@ public:
 	 * @param Subscriber The subscriber to stop routing messages to.
 	 * @param MessageType The type of message to unsubscribe from (NAME_None = all types).
 	 */
-	FORCEINLINE void RemoveSubscription(const TSharedRef<IMessageReceiver, ESPMode::ThreadSafe>& Subscriber, const FName& MessageType)
+	FORCEINLINE void RemoveSubscription(const TSharedRef<IMessageReceiver, ESPMode::ThreadSafe>& Subscriber, const FTopLevelAssetPath& MessageType)
 	{
 		EnqueueCommand(FSimpleDelegate::CreateRaw(this, &FMessageRouter::HandleRemoveSubscriber, TWeakPtr<IMessageReceiver, ESPMode::ThreadSafe>(Subscriber), MessageType));
 	}
@@ -267,7 +267,7 @@ private:
 private:
 
 	/** Handles adding message interceptors. */
-	void HandleAddInterceptor(TSharedRef<IMessageInterceptor, ESPMode::ThreadSafe> Interceptor, FName MessageType);
+	void HandleAddInterceptor(TSharedRef<IMessageInterceptor, ESPMode::ThreadSafe> Interceptor, FTopLevelAssetPath MessageType);
 
 	/** Handles adding message recipients. */
 	void HandleAddRecipient(FMessageAddress Address, TWeakPtr<IMessageReceiver, ESPMode::ThreadSafe> RecipientPtr);
@@ -276,13 +276,13 @@ private:
 	void HandleAddSubscriber(TSharedRef<IMessageSubscription, ESPMode::ThreadSafe> Subscription);
 
 	/** Handles the removal of message interceptors. */
-	void HandleRemoveInterceptor(TSharedRef<IMessageInterceptor, ESPMode::ThreadSafe> Interceptor, FName MessageType);
+	void HandleRemoveInterceptor(TSharedRef<IMessageInterceptor, ESPMode::ThreadSafe> Interceptor, FTopLevelAssetPath MessageType);
 
 	/** Handles the removal of message recipients. */
 	void HandleRemoveRecipient(FMessageAddress Address);
 
 	/** Handles the removal of subscribers. */
-	void HandleRemoveSubscriber(TWeakPtr<IMessageReceiver, ESPMode::ThreadSafe> SubscriberPtr, FName MessageType);
+	void HandleRemoveSubscriber(TWeakPtr<IMessageReceiver, ESPMode::ThreadSafe> SubscriberPtr, FTopLevelAssetPath MessageType);
 
 	/** Handles the routing of messages. */
 	void HandleRouteMessage(TSharedRef<IMessageContext, ESPMode::ThreadSafe> Context);
@@ -299,13 +299,13 @@ private:
 private:
 
 	/** Maps message types to interceptors. */
-	TMap<FName, TArray<TSharedPtr<IMessageInterceptor, ESPMode::ThreadSafe>>> ActiveInterceptors;
+	TMap<FTopLevelAssetPath, TArray<TSharedPtr<IMessageInterceptor, ESPMode::ThreadSafe>>> ActiveInterceptors;
 
 	/** Maps message addresses to recipients. */
 	TMap<FMessageAddress, TWeakPtr<IMessageReceiver, ESPMode::ThreadSafe>> ActiveRecipients;
 
 	/** Maps message types to subscriptions. */
-	TMap<FName, TArray<TSharedPtr<IMessageSubscription, ESPMode::ThreadSafe>>> ActiveSubscriptions;
+	TMap<FTopLevelAssetPath, TArray<TSharedPtr<IMessageSubscription, ESPMode::ThreadSafe>>> ActiveSubscriptions;
 
 	/** Array of active registration listeners. */
 	TArray<TWeakPtr<IBusListener, ESPMode::ThreadSafe>> ActiveRegistrationListeners;
diff --git a/Engine/Source/Runtime/Messaging/Private/Bus/MessageSubscription.h b/Engine/Source/Runtime/Messaging/Private/Bus/MessageSubscription.h
--- a/Engine/Source/Runtime/Messaging/Private/Bus/MessageSubscription.h
+++ b/Engine/Source/Runtime/Messaging/Private/Bus/MessageSubscription.h
@@ -27,7 +27,7 @@ public:
 	 * @param InReceivingThread The thread on which to receive messages on.
 	 * @param InScopeRange The message scope range to subscribe to.
 	 */
-	FMessageSubscription(const TSharedRef<IMessageReceiver, ESPMode::ThreadSafe>& InSubscriber, const FName& InMessageType, const FMessageScopeRange& InScopeRange)
+	FMessageSubscription(const TSharedRef<IMessageReceiver, ESPMode::ThreadSafe>& InSubscriber, const FTopLevelAssetPath& InMessageType, const FMessageScopeRange& InScopeRange)
 		: Enabled(true)
 		, MessageType(InMessageType)
 		, ScopeRange(InScopeRange)
@@ -48,7 +48,7 @@ public:
 		Enabled = true;
 	}
 
-	virtual FName GetMessageType() override
+	virtual FTopLevelAssetPath GetMessageTypePathName() override
 	{
 		return MessageType;
 	}
@@ -74,7 +74,7 @@ private:
 	bool Enabled;
 
 	/** Holds the type of subscribed messages. */
-	FName MessageType;
+	FTopLevelAssetPath MessageType;
 
 	/** Holds the range of message scopes to subscribe to. */
 	FMessageScopeRange ScopeRange;
diff --git a/Engine/Source/Runtime/Messaging/Private/Bus/MessageTracer.cpp b/Engine/Source/Runtime/Messaging/Private/Bus/MessageTracer.cpp
--- a/Engine/Source/Runtime/Messaging/Private/Bus/MessageTracer.cpp
+++ b/Engine/Source/Runtime/Messaging/Private/Bus/MessageTracer.cpp
@@ -31,7 +31,7 @@ FMessageTracer::~FMessageTracer()
 /* FMessageTracer interface
  *****************************************************************************/
 
-void FMessageTracer::TraceAddedInterceptor(const TSharedRef<IMessageInterceptor, ESPMode::ThreadSafe>& Interceptor, const FName& MessageType)
+void FMessageTracer::TraceAddedInterceptor(const TSharedRef<IMessageInterceptor, ESPMode::ThreadSafe>& Interceptor, const FTopLevelAssetPath& MessageType)
 {
 	double Timestamp = FPlatformTime::Seconds();
 
@@ -211,7 +211,7 @@ void FMessageTracer::TraceInterceptedMessage(const TSharedRef<IMessageContext, E
 }
 
 
-void FMessageTracer::TraceRemovedInterceptor(const TSharedRef<IMessageInterceptor, ESPMode::ThreadSafe>& Interceptor, const FName& MessageType)
+void FMessageTracer::TraceRemovedInterceptor(const TSharedRef<IMessageInterceptor, ESPMode::ThreadSafe>& Interceptor, const FTopLevelAssetPath& MessageType)
 {
 	double Timestamp = FPlatformTime::Seconds();
 
@@ -252,7 +252,7 @@ void FMessageTracer::TraceRemovedRecipient(const FMessageAddress& Address)
 }
 
 
-void FMessageTracer::TraceRemovedSubscription(const TSharedRef<IMessageSubscription, ESPMode::ThreadSafe>& Subscription, const FName& MessageType)
+void FMessageTracer::TraceRemovedSubscription(const TSharedRef<IMessageSubscription, ESPMode::ThreadSafe>& Subscription, const FTopLevelAssetPath& MessageType)
 {
 	if (!Running)
 	{
@@ -324,13 +324,12 @@ void FMessageTracer::TraceSentMessage(const TSharedRef<IMessageContext, ESPMode:
 		}
 
 		// add message type
-		TSharedPtr<FMessageTracerTypeInfo>& TypeInfo = MessageTypes.FindOrAdd(Context->GetMessageType());
+		TSharedPtr<FMessageTracerTypeInfo>& TypeInfo = MessageTypes.FindOrAdd(Context->GetMessageTypePathName());
 
 		if (!TypeInfo.IsValid())
 		{
 			TypeInfo = MakeShareable(new FMessageTracerTypeInfo());
-			TypeInfo->TypeName = Context->GetMessageType();
-
+			TypeInfo->TypePathName = Context->GetMessageTypePathName();
 			TypeAddedDelegate.Broadcast(TypeInfo.ToSharedRef());
 		}
 
diff --git a/Engine/Source/Runtime/Messaging/Private/Bus/MessageTracer.h b/Engine/Source/Runtime/Messaging/Private/Bus/MessageTracer.h
--- a/Engine/Source/Runtime/Messaging/Private/Bus/MessageTracer.h
+++ b/Engine/Source/Runtime/Messaging/Private/Bus/MessageTracer.h
@@ -36,7 +36,7 @@ public:
 	 * @param Interceptor The added interceptor.
 	 * @param MessageType The type of messages being intercepted.
 	 */
-	void TraceAddedInterceptor(const TSharedRef<IMessageInterceptor, ESPMode::ThreadSafe>& Interceptor, const FName& MessageType);
+	void TraceAddedInterceptor(const TSharedRef<IMessageInterceptor, ESPMode::ThreadSafe>& Interceptor, const FTopLevelAssetPath& MessageType);
 
 	/**
 	 * Notifies the tracer that a message recipient has been added to the message bus.
@@ -84,7 +84,7 @@ public:
 	 * @param Interceptor The removed interceptor.
 	 * @param MessageType The type of messages that is no longer being intercepted.
 	 */
-	void TraceRemovedInterceptor(const TSharedRef<IMessageInterceptor, ESPMode::ThreadSafe>& Interceptor, const FName& MessageType);
+	void TraceRemovedInterceptor(const TSharedRef<IMessageInterceptor, ESPMode::ThreadSafe>& Interceptor, const FTopLevelAssetPath& MessageType);
 
 	/**
 	 * Notifies the tracer that a recipient has been removed from the message bus.
@@ -99,7 +99,7 @@ public:
 	 * @param Subscriber The removed subscriber.
 	 * @param MessageType The type of messages no longer being subscribed to.
 	 */
-	void TraceRemovedSubscription(const TSharedRef<IMessageSubscription, ESPMode::ThreadSafe>& Subscription, const FName& MessageType);
+	void TraceRemovedSubscription(const TSharedRef<IMessageSubscription, ESPMode::ThreadSafe>& Subscription, const FTopLevelAssetPath& MessageType);
 
 	/**
 	 * Notifies the tracer that a message has been routed.
@@ -187,7 +187,7 @@ private:
 	TMap<TSharedPtr<IMessageContext, ESPMode::ThreadSafe>, TSharedPtr<FMessageTracerMessageInfo>> MessageInfos;
 
 	/** Holds the collection of known message types. */
-	TMap<FName, TSharedPtr<FMessageTracerTypeInfo>> MessageTypes;
+	TMap<FTopLevelAssetPath, TSharedPtr<FMessageTracerTypeInfo>> MessageTypes;
 
 	/** Holds a flag indicating whether a reset is pending. */
 	bool ResetPending;
diff --git a/Engine/Source/Runtime/Messaging/Public/IAuthorizeMessageRecipients.h b/Engine/Source/Runtime/Messaging/Public/IAuthorizeMessageRecipients.h
--- a/Engine/Source/Runtime/Messaging/Public/IAuthorizeMessageRecipients.h
+++ b/Engine/Source/Runtime/Messaging/Public/IAuthorizeMessageRecipients.h
@@ -26,7 +26,20 @@ public:
 	 * @param MessageType The type of messages to intercept.
 	 * @return true if the request was authorized, false otherwise.
 	 */
-	virtual bool AuthorizeInterceptor(const TSharedRef<IMessageInterceptor, ESPMode::ThreadSafe>& Interceptor, const FName& MessageType) = 0;
+	UE_DEPRECATED(5.1, "Types names are now represented by path names. Please use a version of this function that takes an FTopLevelAssetPath as MessageType.")
+	virtual bool AuthorizeInterceptor(const TSharedRef<IMessageInterceptor, ESPMode::ThreadSafe>& Interceptor, const FName& MessageType)
+	{
+		return AuthorizeInterceptor(Interceptor, UClass::TryConvertShortTypeNameToPathName<UStruct>(MessageType.ToString()));
+	}
+
+	/**
+	 * Authorizes a request to intercept messages of the specified type.
+	 *
+	 * @param Interceptor The message interceptor to authorize.
+	 * @param MessageType The type of messages to intercept.
+	 * @return true if the request was authorized, false otherwise.
+	 */
+	virtual bool AuthorizeInterceptor(const TSharedRef<IMessageInterceptor, ESPMode::ThreadSafe>& Interceptor, const FTopLevelAssetPath& MessageType) = 0;
 
 	/**
 	 * Authorizes a request to register the specified recipient.
@@ -44,7 +57,20 @@ public:
 	 * @param TopicPattern The message topic pattern to subscribe to.
 	 * @return true if the request is authorized, false otherwise.
 	 */
-	virtual bool AuthorizeSubscription(const TSharedRef<IMessageReceiver, ESPMode::ThreadSafe>& Subscriber, const FName& TopicPattern) = 0;
+	UE_DEPRECATED(5.1, "Types names are now represented by path names. Please use a version of this function that takes an FTopLevelAssetPath as TopicPattern.")
+	virtual bool AuthorizeSubscription(const TSharedRef<IMessageReceiver, ESPMode::ThreadSafe>& Subscriber, const FName& TopicPattern)
+	{
+		return AuthorizeSubscription(Subscriber, UClass::TryConvertShortTypeNameToPathName<UStruct>(TopicPattern.ToString()));
+	}
+
+	/**
+	 * Authorizes a request to add a subscription for the specified topic pattern.
+	 *
+	 * @param Subscriber The subscriber.
+	 * @param TopicPattern The message topic pattern to subscribe to.
+	 * @return true if the request is authorized, false otherwise.
+	 */
+	virtual bool AuthorizeSubscription(const TSharedRef<IMessageReceiver, ESPMode::ThreadSafe>& Subscriber, const FTopLevelAssetPath& TopicPattern) = 0;
 
 	/**
 	 * Authorizes a request to unregister the specified recipient.
@@ -61,7 +87,20 @@ public:
 	 * @param TopicPattern The message topic pattern to unsubscribe from.
 	 * @return true if the request is authorized, false otherwise.
 	 */
-	virtual bool AuthorizeUnsubscription(const TSharedRef<IMessageReceiver, ESPMode::ThreadSafe>& Subscriber, const FName& TopicPattern) = 0;
+	UE_DEPRECATED(5.1, "Types names are now represented by path names. Please use a version of this function that takes an FTopLevelAssetPath as TopicPattern.")
+	virtual bool AuthorizeUnsubscription(const TSharedRef<IMessageReceiver, ESPMode::ThreadSafe>& Subscriber, const FName& TopicPattern)
+	{
+		return AuthorizeUnsubscription(Subscriber, UClass::TryConvertShortTypeNameToPathName<UStruct>(TopicPattern.ToString()));
+	}
+
+	/**
+	 * Authorizes a request to remove a subscription for the specified topic pattern.
+	 *
+	 * @param Subscriber The subscriber.
+	 * @param TopicPattern The message topic pattern to unsubscribe from.
+	 * @return true if the request is authorized, false otherwise.
+	 */
+	virtual bool AuthorizeUnsubscription(const TSharedRef<IMessageReceiver, ESPMode::ThreadSafe>& Subscriber, const FTopLevelAssetPath& TopicPattern) = 0;
 
 public:
 
diff --git a/Engine/Source/Runtime/Messaging/Public/IMessageBus.h b/Engine/Source/Runtime/Messaging/Public/IMessageBus.h
--- a/Engine/Source/Runtime/Messaging/Public/IMessageBus.h
+++ b/Engine/Source/Runtime/Messaging/Public/IMessageBus.h
@@ -140,7 +140,20 @@ public:
 	 * @param MessageType The type of messages to intercept.
 	 * @see Unintercept
 	 */
-	virtual void Intercept(const TSharedRef<IMessageInterceptor, ESPMode::ThreadSafe>& Interceptor, const FName& MessageType) = 0;
+	UE_DEPRECATED(5.1, "Types names are now represented by path names. Please use a version of this function that takes an FTopLevelAssetPath as MessageType.")
+	virtual void Intercept(const TSharedRef<IMessageInterceptor, ESPMode::ThreadSafe>& Interceptor, const FName& MessageType)
+	{
+		Intercept(Interceptor, UClass::TryConvertShortTypeNameToPathName<UStruct>(MessageType.ToString()));
+	}
+
+	/**
+	 * Adds an interceptor for messages of the specified type.
+	 *
+	 * @param Interceptor The interceptor.
+	 * @param MessageType The type of messages to intercept.
+	 * @see Unintercept
+	 */
+	virtual void Intercept(const TSharedRef<IMessageInterceptor, ESPMode::ThreadSafe>& Interceptor, const FTopLevelAssetPath& MessageType) = 0;
 
 	/**
 	 * Sends a message to subscribed recipients.
@@ -206,7 +219,38 @@ public:
 	 * @return The added subscription, or nullptr if the subscription failed.
 	 * @see Unsubscribe
 	 */
-	virtual TSharedPtr<IMessageSubscription, ESPMode::ThreadSafe> Subscribe(const TSharedRef<IMessageReceiver, ESPMode::ThreadSafe>& Subscriber, const FName& MessageType, const TRange<EMessageScope>& ScopeRange) = 0;
+	UE_DEPRECATED(5.1, "Types names are now represented by path names. Please use a version of this function that takes an FTopLevelAssetPath as MessageType.")
+	virtual TSharedPtr<IMessageSubscription, ESPMode::ThreadSafe> Subscribe(const TSharedRef<IMessageReceiver, ESPMode::ThreadSafe>& Subscriber, const FName& MessageType, const TRange<EMessageScope>& ScopeRange)
+	{
+		return Subscribe(Subscriber, UClass::TryConvertShortTypeNameToPathName<UStruct>(MessageType.ToString()), ScopeRange);
+	}
+
+	/**
+	 * Adds a subscription for published messages of the specified type.
+	 *
+	 * Subscriptions allow message consumers to receive published messages from the message bus.
+	 * The returned interface can be used to query the subscription's details and its enabled state.
+	 *
+	 * @param Subscriber The subscriber wishing to receive the messages.
+	 * @param MessageType The type of messages to subscribe to (IMessageBus::PATHNAME_All = subscribe to all message types).
+	 * @param ScopeRange The range of message scopes to include in the subscription.
+	 * @return The added subscription, or nullptr if the subscription failed.
+	 * @see Unsubscribe
+	 */
+	virtual TSharedPtr<IMessageSubscription, ESPMode::ThreadSafe> Subscribe(const TSharedRef<IMessageReceiver, ESPMode::ThreadSafe>& Subscriber, const FTopLevelAssetPath& MessageType, const TRange<EMessageScope>& ScopeRange) = 0;
+
+	/**
+	 * Removes an interceptor for messages of the specified type.
+	 *
+	 * @param Interceptor The interceptor to remove.
+	 * @param MessageType The type of messages to stop intercepting.
+	 * @see Intercept
+	 */
+	UE_DEPRECATED(5.1, "Types names are now represented by path names. Please use a version of this function that takes an FTopLevelAssetPath as MessageType.")
+	virtual void Unintercept(const TSharedRef<IMessageInterceptor, ESPMode::ThreadSafe>& Interceptor, const FName& MessageType)
+	{
+		Unintercept(Interceptor, UClass::TryConvertShortTypeNameToPathName<UStruct>(MessageType.ToString()));
+	}
 
 	/**
 	 * Removes an interceptor for messages of the specified type.
@@ -215,7 +259,7 @@ public:
 	 * @param MessageType The type of messages to stop intercepting.
 	 * @see Intercept
 	 */
-	virtual void Unintercept(const TSharedRef<IMessageInterceptor, ESPMode::ThreadSafe>& Interceptor, const FName& MessageType) = 0;
+	virtual void Unintercept(const TSharedRef<IMessageInterceptor, ESPMode::ThreadSafe>& Interceptor, const FTopLevelAssetPath& MessageType) = 0;
 
 	/**
 	 * Unregisters a message recipient from the message bus.
@@ -232,7 +276,20 @@ public:
 	 * @param MessageType The type of messages to unsubscribe from (NAME_All = all types).
 	 * @see Subscribe
 	 */
-	virtual void Unsubscribe(const TSharedRef<IMessageReceiver, ESPMode::ThreadSafe>& Subscriber, const FName& MessageType) = 0;
+	UE_DEPRECATED(5.1, "Types names are now represented by path names. Please use a version of this function that takes an FTopLevelAssetPath as MessageType.")
+	virtual void Unsubscribe(const TSharedRef<IMessageReceiver, ESPMode::ThreadSafe>& Subscriber, const FName& MessageType)
+	{
+		Unsubscribe(Subscriber, UClass::TryConvertShortTypeNameToPathName<UStruct>(MessageType.ToString()));
+	}
+
+	/**
+	 * Cancels the specified message subscription.
+	 *
+	 * @param Subscriber The subscriber wishing to stop receiving the messages.
+	 * @param MessageType The type of messages to unsubscribe from (IMessageBus::PATHNAME_All = all types).
+	 * @see Subscribe
+	 */
+	virtual void Unsubscribe(const TSharedRef<IMessageReceiver, ESPMode::ThreadSafe>& Subscriber, const FTopLevelAssetPath& MessageType) = 0;
 
 	/**
 	 * Add a listener to the bus notifications
@@ -261,6 +318,8 @@ public:
 	 */
 	virtual FOnMessageBusShutdown& OnShutdown() = 0;
 
+	MESSAGING_API static const FTopLevelAssetPath PATHNAME_All;
+
 public:
 
 	/** Virtual destructor. */
diff --git a/Engine/Source/Runtime/Messaging/Public/IMessageContext.h b/Engine/Source/Runtime/Messaging/Public/IMessageContext.h
--- a/Engine/Source/Runtime/Messaging/Public/IMessageContext.h
+++ b/Engine/Source/Runtime/Messaging/Public/IMessageContext.h
@@ -335,6 +335,7 @@ public:
 	 * @return Message type name.
 	 * @see GetMessage, GetMessageTypeInfo
 	 */
+	UE_DEPRECATED(5.1, "Type names are now represented by path names. Please use GetMessageTypePathName.")
 	FName GetMessageType() const
 	{
 		if (IsValid())
@@ -346,6 +347,23 @@ public:
 		return NAME_None;
 	}
 
+	/**
+	 * Gets the path name of the message type.
+	 *
+	 * @return Message type path name.
+	 * @see GetMessage, GetMessageTypeInfo
+	 */
+	FTopLevelAssetPath GetMessageTypePathName() const
+	{
+		if (IsValid())
+		{
+			UStruct* MessageTypeInfoPtr = GetMessageTypeInfo().Get();
+			return MessageTypeInfoPtr->GetStructPathName();
+		}
+
+		return FTopLevelAssetPath();
+	}
+
 	/**
 	 * Checks whether this is a forwarded message.
 	 *
diff --git a/Engine/Source/Runtime/Messaging/Public/IMessageSubscription.h b/Engine/Source/Runtime/Messaging/Public/IMessageSubscription.h
--- a/Engine/Source/Runtime/Messaging/Public/IMessageSubscription.h
+++ b/Engine/Source/Runtime/Messaging/Public/IMessageSubscription.h
@@ -44,7 +44,19 @@ public:
 	 * @return Message type.
 	 * @see GetScopeRange, GetSubscriber
 	 */
-	virtual FName GetMessageType() = 0;
+	UE_DEPRECATED(5.1, "Types names are now represented by path names. Please use GetMessageTypePathName.")
+	FORCEINLINE FName GetMessageType()
+	{
+		return GetMessageTypePathName().GetAssetName();
+	}
+
+	/**
+	 * Gets the type of subscribed messages.
+	 *
+	 * @return Message type.
+	 * @see GetScopeRange, GetSubscriber
+	 */
+	virtual FTopLevelAssetPath GetMessageTypePathName() = 0;
 
 	/**
 	 * Gets the range of subscribed message scopes.
diff --git a/Engine/Source/Runtime/Messaging/Public/IMessageTracer.h b/Engine/Source/Runtime/Messaging/Public/IMessageTracer.h
--- a/Engine/Source/Runtime/Messaging/Public/IMessageTracer.h
+++ b/Engine/Source/Runtime/Messaging/Public/IMessageTracer.h
@@ -188,7 +188,11 @@ struct FMessageTracerTypeInfo
 	TArray<TSharedPtr<FMessageTracerMessageInfo>> Messages;
 
 	/** Holds a name of the message type. */
+	UE_DEPRECATED(5.1, "Types names are now represented by path names. Please use TypePathName.")
 	FName TypeName;
+
+	/** Holds a path name of the message type. */
+	FTopLevelAssetPath TypePathName;
 };
 
 
diff --git a/Engine/Source/Runtime/MessagingCommon/Public/MessageEndpoint.h b/Engine/Source/Runtime/MessagingCommon/Public/MessageEndpoint.h
--- a/Engine/Source/Runtime/MessagingCommon/Public/MessageEndpoint.h
+++ b/Engine/Source/Runtime/MessagingCommon/Public/MessageEndpoint.h
@@ -332,7 +332,19 @@ public:
 	 * @param MessageType The type name of the messages to subscribe to.
 	 * @param ScopeRange The range of message scopes to include in the subscription.
 	 */
+	UE_DEPRECATED(5.1, "Types names are now represented by path names. Please use a version of this function that takes an FTopLevelAssetPath as MessageType.")
 	void Subscribe(const FName& MessageType, const FMessageScopeRange& ScopeRange)
+	{
+		Subscribe(UClass::TryConvertShortTypeNameToPathName<UStruct>(MessageType.ToString()), ScopeRange);
+	}
+
+	/**
+	 * Subscribes a message handler.
+	 *
+	 * @param MessageType The type name of the messages to subscribe to.
+	 * @param ScopeRange The range of message scopes to include in the subscription.
+	 */
+	void Subscribe(const FTopLevelAssetPath& MessageType, const FMessageScopeRange& ScopeRange)
 	{
 		TSharedPtr<IMessageBus, ESPMode::ThreadSafe> Bus = GetBusIfEnabled();
 
@@ -345,16 +357,28 @@ public:
 	/**
 	 * Unsubscribes this endpoint from the specified message type.
 	 *
-	 * @param MessageType The type of message to unsubscribe (NAME_All = all types).
+	 * @param MessageType The type of message to unsubscribe (IMessageBus::PATHNAME_All = all types).
+	 * @see Subscribe
+	 */
+	UE_DEPRECATED(5.1, "Types names are now represented by path names. Please use a version of this function that takes an FTopLevelAssetPath as MessageType.")
+	void Unsubscribe(const FName& MessageType)
+	{
+		Unsubscribe(UClass::TryConvertShortTypeNameToPathName<UStruct>(MessageType.ToString()));
+	}
+
+	/**
+	 * Unsubscribes this endpoint from the specified message type.
+	 *
+	 * @param MessageType The type of message to unsubscribe (IMessageBus::PATHNAME_All = all types).
 	 * @see Subscribe
 	 */
-	void Unsubscribe(const FName& TopicPattern)
+	void Unsubscribe(const FTopLevelAssetPath& MessageType)
 	{
 		TSharedPtr<IMessageBus, ESPMode::ThreadSafe> Bus = GetBusIfEnabled();
 
 		if (Bus.IsValid())
 		{
-			Bus->Unsubscribe(AsShared(), TopicPattern);
+			Bus->Unsubscribe(AsShared(), MessageType);
 		}
 	}
 
@@ -890,7 +914,7 @@ public:
 	template<class MessageType>
 	void Subscribe()
 	{
-		Subscribe(MessageType::StaticStruct()->GetFName(), FMessageScopeRange::AtLeast(EMessageScope::Thread));
+		Subscribe(MessageType::StaticStruct()->GetStructPathName(), FMessageScopeRange::AtLeast(EMessageScope::Thread));
 	}
 
 	/**
@@ -905,7 +929,7 @@ public:
 	template<class MessageType>
 	void Subscribe(const FMessageScopeRange& ScopeRange)
 	{
-		Subscribe(MessageType::StaticStruct()->GetFName(), ScopeRange);
+		Subscribe(MessageType::StaticStruct()->GetStructPathName(), ScopeRange);
 	}
 
 	/**
@@ -915,19 +939,19 @@ public:
 	 */
 	void Unsubscribe()
 	{
-		Unsubscribe(NAME_All);
+		Unsubscribe(IMessageBus::PATHNAME_All);
 	}
 
 	/**
 	 * Template method to unsubscribe the endpoint from the specified message type.
 	 *
-	 * @param MessageType The type of message to unsubscribe (NAME_All = all types).
+	 * @param MessageType The type of message to unsubscribe (IMessageBus::PATHNAME_All = all types).
 	 * @see Subscribe
 	 */
 	template<class MessageType>
 	void Unsubscribe()
 	{
-		Unsubscribe(MessageType::StaticStruct()->GetFName());
+		Unsubscribe(MessageType::StaticStruct()->GetStructPathName());
 	}
 
 public:
diff --git a/Engine/Source/Runtime/MessagingCommon/Public/MessageHandlers.h b/Engine/Source/Runtime/MessagingCommon/Public/MessageHandlers.h
--- a/Engine/Source/Runtime/MessagingCommon/Public/MessageHandlers.h
+++ b/Engine/Source/Runtime/MessagingCommon/Public/MessageHandlers.h
@@ -98,7 +98,7 @@ public:
 	virtual void HandleMessage(const TSharedRef<IMessageContext, ESPMode::ThreadSafe>& Context) override
 	{
 		UStruct* Struct = MessageType::StaticStruct();
-		if (Struct && Context->GetMessageType() == Struct->GetFName())
+		if (Struct && Context->GetMessageTypePathName() == Struct->GetStructPathName())
 		{
 			(Handler->*Func)(*static_cast<const MessageType*>(Context->GetMessage()), Context);
 		}	
@@ -189,7 +189,7 @@ public:
 	
 	virtual void HandleMessage(const TSharedRef<IMessageContext, ESPMode::ThreadSafe>& Context) override
 	{
-		if (Context->GetMessageType() == MessageType::StaticStruct()->GetFName())
+		if (Context->GetMessageTypePathName() == MessageType::StaticStruct()->GetStructPathName())
 		{
 			Func(*static_cast<const MessageType*>(Context->GetMessage()), Context);
 		}	
diff --git a/Engine/Source/Runtime/MessagingRpc/Private/MessageRpcServer.cpp b/Engine/Source/Runtime/MessagingRpc/Private/MessageRpcServer.cpp
--- a/Engine/Source/Runtime/MessagingRpc/Private/MessageRpcServer.cpp
+++ b/Engine/Source/Runtime/MessagingRpc/Private/MessageRpcServer.cpp
@@ -47,7 +47,7 @@ TSharedPtr<FMessageEndpoint, ESPMode::ThreadSafe> FMessageRpcServer::GetEndpoint
 	return MessageEndpoint;
 }
 
-void FMessageRpcServer::AddHandler(const FName& RequestMessageType, const TSharedRef<IMessageRpcHandler>& Handler)
+void FMessageRpcServer::AddHandler(const FTopLevelAssetPath& RequestMessageType, const TSharedRef<IMessageRpcHandler>& Handler)
 {
 	Handlers.Add(RequestMessageType, Handler);
 }
@@ -59,11 +59,17 @@ const FMessageAddress& FMessageRpcServer::GetAddress() const
 }
 
 
+PRAGMA_DISABLE_DEPRECATION_WARNINGS
 FOnMessageRpcNoHandler& FMessageRpcServer::OnNoHandler()
 {
 	return NoHandlerDelegate;
 }
+PRAGMA_ENABLE_DEPRECATION_WARNINGS
 
+FOnMessagePathNameRpcNoHandler& FMessageRpcServer::OnNoHandlerWithPathName()
+{
+	return NoHandlerDelegateWithPathName;
+}
 
 void FMessageRpcServer::SetSendProgressUpdate(bool InSendProgress)
 {
@@ -87,16 +93,20 @@ void FMessageRpcServer::ProcessCancelation(const FMessageRpcCancel& Message, con
 void FMessageRpcServer::ProcessRequest(const TSharedRef<IMessageContext, ESPMode::ThreadSafe>& Context)
 {
 	auto Message = (FRpcMessage*)Context->GetMessage();
-	const FName MessageType = Context->GetMessageType();
+	const FTopLevelAssetPath MessageType = Context->GetMessageTypePathName();
 
 	if (!Handlers.Contains(MessageType))
 	{
-		if (!NoHandlerDelegate.IsBound())
+PRAGMA_DISABLE_DEPRECATION_WARNINGS
+		// Keeping checking and executing NoHandlerDelegate for backwards compatibility 
+		if (!NoHandlerDelegateWithPathName.IsBound() && !NoHandlerDelegate.IsBound())
 		{
 			return;
 		}
 
-		NoHandlerDelegate.Execute(MessageType);
+		NoHandlerDelegateWithPathName.ExecuteIfBound(MessageType);
+		NoHandlerDelegate.ExecuteIfBound(MessageType.GetAssetName());
+PRAGMA_DISABLE_DEPRECATION_WARNINGS
 	}
 
 	auto Handler = Handlers.FindRef(MessageType);
diff --git a/Engine/Source/Runtime/MessagingRpc/Public/IMessageRpcServer.h b/Engine/Source/Runtime/MessagingRpc/Public/IMessageRpcServer.h
--- a/Engine/Source/Runtime/MessagingRpc/Public/IMessageRpcServer.h
+++ b/Engine/Source/Runtime/MessagingRpc/Public/IMessageRpcServer.h
@@ -14,8 +14,11 @@ struct FRpcMessage;
 template<typename ResultType> class TFuture;
 
 /** Delegate type for RPC messages that have no registered handler. */
+UE_DEPRECATED(5.1, "Types names are now represented by path names. Please use FOnMessagePathNameRpcNoHandler.")
 DECLARE_DELEGATE_OneParam(FOnMessageRpcNoHandler, const FName& /*MessageType*/)
 
+DECLARE_DELEGATE_OneParam(FOnMessagePathNameRpcNoHandler, const FTopLevelAssetPath& /*MessageType*/)
+
 /**
  * Interface for RPC servers.
  */
@@ -98,7 +101,18 @@ public:
 	 *
 	 * @param Handler The handler to add.
 	 */
-	virtual void AddHandler(const FName& RequestMessageType, const TSharedRef<IMessageRpcHandler>& Handler) = 0;
+	UE_DEPRECATED(5.1, "Types names are now represented by path names. Please use a version of this function that takes an FTopLevelAssetPath as MessageType.")
+	FORCEINLINE void AddHandler(const FName& RequestMessageType, const TSharedRef<IMessageRpcHandler>& Handler)
+	{
+		AddHandler(UClass::TryConvertShortTypeNameToPathName<UStruct>(RequestMessageType.ToString()), Handler);
+	}
+
+	/**
+	 * Add an RPC request handler.
+	 *
+	 * @param Handler The handler to add.
+	 */
+	virtual void AddHandler(const FTopLevelAssetPath& RequestMessageType, const TSharedRef<IMessageRpcHandler>& Handler) = 0;
 
 	/**
 	 * Gets the server's message address.
@@ -112,8 +126,17 @@ public:
 	 *
 	 * @return The delegate.
 	 */
+	PRAGMA_DISABLE_DEPRECATION_WARNINGS
+	UE_DEPRECATED(5.1, "Types names are now represented by path names. Please use OnNoHandlerWithPathName.")	
 	virtual FOnMessageRpcNoHandler& OnNoHandler() = 0;
+	PRAGMA_ENABLE_DEPRECATION_WARNINGS
 
+	/**
+	 * Get a delegate that is executed when a received RPC message has no registered handler.
+	 *
+	 * @return The delegate.
+	 */
+	virtual FOnMessagePathNameRpcNoHandler& OnNoHandlerWithPathName() = 0;
 
 	/** 
 	 * Set if the server sends progress updates
diff --git a/Engine/Source/Runtime/MessagingRpc/Public/MessageRpcServer.h b/Engine/Source/Runtime/MessagingRpc/Public/MessageRpcServer.h
--- a/Engine/Source/Runtime/MessagingRpc/Public/MessageRpcServer.h
+++ b/Engine/Source/Runtime/MessagingRpc/Public/MessageRpcServer.h
@@ -39,9 +39,12 @@ public:
 
 	//~ IMessageRpcServer interface
 
-	virtual void AddHandler(const FName& RequestMessageType, const TSharedRef<IMessageRpcHandler>& Handler) override;
+	virtual void AddHandler(const FTopLevelAssetPath& RequestMessageType, const TSharedRef<IMessageRpcHandler>& Handler) override;
 	virtual const FMessageAddress& GetAddress() const override;
+PRAGMA_DISABLE_DEPRECATION_WARNINGS
 	virtual FOnMessageRpcNoHandler& OnNoHandler() override;
+PRAGMA_ENABLE_DEPRECATION_WARNINGS
+	virtual FOnMessagePathNameRpcNoHandler& OnNoHandlerWithPathName() override;
 	virtual void SetSendProgressUpdate(bool InSendProgress) override;
 protected:
 	explicit FMessageRpcServer(FMessageEndpointBuilder&& InEndpointBuilder);
@@ -85,10 +88,15 @@ private:
 private:
 
 	/** Registered request message handlers. */
-	TMap<FName, TSharedPtr<IMessageRpcHandler>> Handlers;
+	TMap<FTopLevelAssetPath, TSharedPtr<IMessageRpcHandler>> Handlers;
 
 	/* Delegate that is executed when a received RPC message has no registered handler. */
+PRAGMA_DISABLE_DEPRECATION_WARNINGS
 	FOnMessageRpcNoHandler NoHandlerDelegate;
+PRAGMA_ENABLE_DEPRECATION_WARNINGS
+
+	/* Delegate that is executed when a received RPC message has no registered handler. */
+	FOnMessagePathNameRpcNoHandler NoHandlerDelegateWithPathName;
 
 	/** Collection of pending RPC returns. */
 	TMap<FGuid, FReturnInfo> Returns;
diff --git a/Engine/Source/Runtime/MovieSceneCapture/Private/MovieSceneCapture.cpp b/Engine/Source/Runtime/MovieSceneCapture/Private/MovieSceneCapture.cpp
--- a/Engine/Source/Runtime/MovieSceneCapture/Private/MovieSceneCapture.cpp
+++ b/Engine/Source/Runtime/MovieSceneCapture/Private/MovieSceneCapture.cpp
@@ -38,7 +38,7 @@ TArray<UClass*> FindAllCaptureProtocolClasses()
 	TArray<FAssetData> BlueprintList;
 
 	FARFilter Filter;
-	Filter.ClassNames.Add(UMovieSceneCaptureProtocolBase::StaticClass()->GetFName());
+	Filter.ClassPaths.Add(UMovieSceneCaptureProtocolBase::StaticClass()->GetClassPathName());
 
 	// Include any Blueprint based objects as well, this includes things like Blutilities, UMG, and GameplayAbility objects
 	Filter.bRecursiveClasses = true;
diff --git a/Engine/Source/Runtime/PacketHandlers/PacketHandler/Private/PacketHandler.cpp b/Engine/Source/Runtime/PacketHandlers/PacketHandler/Private/PacketHandler.cpp
--- a/Engine/Source/Runtime/PacketHandlers/PacketHandler/Private/PacketHandler.cpp
+++ b/Engine/Source/Runtime/PacketHandlers/PacketHandler/Private/PacketHandler.cpp
@@ -346,7 +346,7 @@ TSharedPtr<HandlerComponent> PacketHandler::AddHandler(const FString& ComponentS
 			{
 				// Every HandlerComponentFactory type has one instance, loaded as a named singleton
 				FString SingletonName = ComponentName.Mid(FactoryComponentDelim + 1) + TEXT("_Singleton");
-				UHandlerComponentFactory* Factory = FindObject<UHandlerComponentFactory>(ANY_PACKAGE, *SingletonName);
+				UHandlerComponentFactory* Factory = FindFirstObject<UHandlerComponentFactory>(*SingletonName, EFindFirstObjectOptions::NativeFirst | EFindFirstObjectOptions::EnsureIfAmbiguous);
 
 				if (Factory == nullptr)
 				{
diff --git a/Engine/Source/Runtime/RigVM/Private/RigVMCore/RigVM.cpp b/Engine/Source/Runtime/RigVM/Private/RigVMCore/RigVM.cpp
--- a/Engine/Source/Runtime/RigVM/Private/RigVMCore/RigVM.cpp
+++ b/Engine/Source/Runtime/RigVM/Private/RigVMCore/RigVM.cpp
@@ -60,7 +60,7 @@ UScriptStruct* FRigVMParameter::GetScriptStruct() const
 		if (ScriptStructPath != NAME_None)
 		{
 			FRigVMParameter* MutableThis = (FRigVMParameter*)this;
-			MutableThis->ScriptStruct = FindObject<UScriptStruct>(ANY_PACKAGE, *ScriptStructPath.ToString());
+			MutableThis->ScriptStruct = FindObject<UScriptStruct>(nullptr, *ScriptStructPath.ToString());
 		}
 	}
 	return ScriptStruct;
diff --git a/Engine/Source/Runtime/RigVM/Private/RigVMCore/RigVMMemoryDeprecated.cpp b/Engine/Source/Runtime/RigVM/Private/RigVMCore/RigVMMemoryDeprecated.cpp
--- a/Engine/Source/Runtime/RigVM/Private/RigVMCore/RigVMMemoryDeprecated.cpp
+++ b/Engine/Source/Runtime/RigVM/Private/RigVMCore/RigVMMemoryDeprecated.cpp
@@ -123,7 +123,7 @@ void FRigVMRegisterOffset::Load(FArchive& Ar)
 		FName ScriptStructPath;
 		Ar << ScriptStructPath;
 
-		ScriptStruct = FindObject<UScriptStruct>(ANY_PACKAGE, *ScriptStructPath.ToString());
+		ScriptStruct = UClass::TryFindTypeSlow<UScriptStruct>(ScriptStructPath.ToString());
 	}
 	else
 	{
