Fixing non-unity build

diff --git a/Engine/Source/Runtime/Messaging/Private/Bus/MessageBus.cpp b/Engine/Source/Runtime/Messaging/Private/Bus/MessageBus.cpp
--- a/Engine/Source/Runtime/Messaging/Private/Bus/MessageBus.cpp
+++ b/Engine/Source/Runtime/Messaging/Private/Bus/MessageBus.cpp
@@ -240,3 +240,25 @@ const FString& FMessageBus::GetName() const
 {
 	return Name;
 }
+
+// IMessageBus deprecated functions
+
+void IMessageBus::Intercept(const TSharedRef<IMessageInterceptor, ESPMode::ThreadSafe>& Interceptor, const FName& MessageType)
+{
+	Intercept(Interceptor, UClass::TryConvertShortTypeNameToPathName<UStruct>(MessageType.ToString()));
+}
+
+TSharedPtr<IMessageSubscription, ESPMode::ThreadSafe> IMessageBus::Subscribe(const TSharedRef<IMessageReceiver, ESPMode::ThreadSafe>& Subscriber, const FName& MessageType, const TRange<EMessageScope>& ScopeRange)
+{
+	return Subscribe(Subscriber, UClass::TryConvertShortTypeNameToPathName<UStruct>(MessageType.ToString()), ScopeRange);
+}
+
+void IMessageBus::Unintercept(const TSharedRef<IMessageInterceptor, ESPMode::ThreadSafe>& Interceptor, const FName& MessageType)
+{
+	Unintercept(Interceptor, UClass::TryConvertShortTypeNameToPathName<UStruct>(MessageType.ToString()));
+}
+
+void IMessageBus::Unsubscribe(const TSharedRef<IMessageReceiver, ESPMode::ThreadSafe>& Subscriber, const FName& MessageType)
+{
+	Unsubscribe(Subscriber, UClass::TryConvertShortTypeNameToPathName<UStruct>(MessageType.ToString()));
+}
\ No newline at end of file
diff --git a/Engine/Source/Runtime/Messaging/Public/IMessageBus.h b/Engine/Source/Runtime/Messaging/Public/IMessageBus.h
--- a/Engine/Source/Runtime/Messaging/Public/IMessageBus.h
+++ b/Engine/Source/Runtime/Messaging/Public/IMessageBus.h
@@ -5,6 +5,7 @@
 #include "Containers/Array.h"
 #include "Math/Range.h"
 #include "Templates/SharedPointer.h"
+#include "UObject/TopLevelAssetPath.h"
 
 class FName;
 class IMessageAttachment;
@@ -25,7 +26,6 @@ struct FDateTime;
 struct FMessageAddress;
 struct FTimespan;
 
-
 /** Delegate type for message bus shutdowns. */
 DECLARE_MULTICAST_DELEGATE(FOnMessageBusShutdown);
 
@@ -141,10 +141,7 @@ public:
 	 * @see Unintercept
 	 */
 	UE_DEPRECATED(5.1, "Types names are now represented by path names. Please use a version of this function that takes an FTopLevelAssetPath as MessageType.")
-	virtual void Intercept(const TSharedRef<IMessageInterceptor, ESPMode::ThreadSafe>& Interceptor, const FName& MessageType)
-	{
-		Intercept(Interceptor, UClass::TryConvertShortTypeNameToPathName<UStruct>(MessageType.ToString()));
-	}
+	MESSAGING_API virtual void Intercept(const TSharedRef<IMessageInterceptor, ESPMode::ThreadSafe>& Interceptor, const FName& MessageType);
 
 	/**
 	 * Adds an interceptor for messages of the specified type.
@@ -220,10 +217,7 @@ public:
 	 * @see Unsubscribe
 	 */
 	UE_DEPRECATED(5.1, "Types names are now represented by path names. Please use a version of this function that takes an FTopLevelAssetPath as MessageType.")
-	virtual TSharedPtr<IMessageSubscription, ESPMode::ThreadSafe> Subscribe(const TSharedRef<IMessageReceiver, ESPMode::ThreadSafe>& Subscriber, const FName& MessageType, const TRange<EMessageScope>& ScopeRange)
-	{
-		return Subscribe(Subscriber, UClass::TryConvertShortTypeNameToPathName<UStruct>(MessageType.ToString()), ScopeRange);
-	}
+	MESSAGING_API virtual TSharedPtr<IMessageSubscription, ESPMode::ThreadSafe> Subscribe(const TSharedRef<IMessageReceiver, ESPMode::ThreadSafe>& Subscriber, const FName& MessageType, const TRange<EMessageScope>& ScopeRange);
 
 	/**
 	 * Adds a subscription for published messages of the specified type.
@@ -247,10 +241,7 @@ public:
 	 * @see Intercept
 	 */
 	UE_DEPRECATED(5.1, "Types names are now represented by path names. Please use a version of this function that takes an FTopLevelAssetPath as MessageType.")
-	virtual void Unintercept(const TSharedRef<IMessageInterceptor, ESPMode::ThreadSafe>& Interceptor, const FName& MessageType)
-	{
-		Unintercept(Interceptor, UClass::TryConvertShortTypeNameToPathName<UStruct>(MessageType.ToString()));
-	}
+	MESSAGING_API virtual void Unintercept(const TSharedRef<IMessageInterceptor, ESPMode::ThreadSafe>& Interceptor, const FName& MessageType);
 
 	/**
 	 * Removes an interceptor for messages of the specified type.
@@ -277,10 +268,7 @@ public:
 	 * @see Subscribe
 	 */
 	UE_DEPRECATED(5.1, "Types names are now represented by path names. Please use a version of this function that takes an FTopLevelAssetPath as MessageType.")
-	virtual void Unsubscribe(const TSharedRef<IMessageReceiver, ESPMode::ThreadSafe>& Subscriber, const FName& MessageType)
-	{
-		Unsubscribe(Subscriber, UClass::TryConvertShortTypeNameToPathName<UStruct>(MessageType.ToString()));
-	}
+	MESSAGING_API virtual void Unsubscribe(const TSharedRef<IMessageReceiver, ESPMode::ThreadSafe>& Subscriber, const FName& MessageType);
 
 	/**
 	 * Cancels the specified message subscription.
