switched from `TQueue` to `TMpscQueue` as it's a bit faster and more modern implementation, and TSAN doesn't complain about it

diff --git a/Engine/Source/Runtime/Messaging/Private/Bus/MessageRouter.h b/Engine/Source/Runtime/Messaging/Private/Bus/MessageRouter.h
--- a/Engine/Source/Runtime/Messaging/Private/Bus/MessageRouter.h
+++ b/Engine/Source/Runtime/Messaging/Private/Bus/MessageRouter.h
@@ -7,6 +7,7 @@
 #include "HAL/Runnable.h"
 #include "Misc/SingleThreadRunnable.h"
 #include "Templates/Atomic.h"
+#include "Containers/MpscQueue.h"
 
 #include "IMessageContext.h"
 #include "IMessageTracer.h"
@@ -170,11 +171,7 @@ protected:
 	 */
 	FORCEINLINE bool EnqueueCommand(CommandDelegate Command)
 	{
-		if (!Commands.Enqueue(Command))
-		{
-			return false;
-		}
-
+		Commands.Enqueue(Command);
 		WorkEvent->Trigger();
 
 		return true;
@@ -311,7 +308,7 @@ private:
 	TArray<TWeakPtr<IBusListener, ESPMode::ThreadSafe>> ActiveRegistrationListeners;
 
 	/** Holds the router command queue. */
-	TQueue<CommandDelegate, EQueueMode::Mpsc> Commands;
+	TMpscQueue<CommandDelegate> Commands;
 
 	/** Holds the current time. */
 	FDateTime CurrentTime;
