Use a thread safe delegate.  This will prevent race conditions where we are adding a delegate handler while checking delegate status.

diff --git a/Engine/Source/Runtime/Messaging/Public/INetworkMessagingExtension.h b/Engine/Source/Runtime/Messaging/Public/INetworkMessagingExtension.h
--- a/Engine/Source/Runtime/Messaging/Public/INetworkMessagingExtension.h
+++ b/Engine/Source/Runtime/Messaging/Public/INetworkMessagingExtension.h
@@ -86,8 +86,8 @@ struct FInboundTransferStatistics
 };
 
 
-DECLARE_MULTICAST_DELEGATE_OneParam(FOnOutboundTransferDataUpdated, FOutboundTransferStatistics);
-DECLARE_MULTICAST_DELEGATE_OneParam(FOnInboundTransferDataUpdated, FInboundTransferStatistics);
+DECLARE_TS_MULTICAST_DELEGATE_OneParam(FOnOutboundTransferDataUpdated, FOutboundTransferStatistics);
+DECLARE_TS_MULTICAST_DELEGATE_OneParam(FOnInboundTransferDataUpdated, FInboundTransferStatistics);
 
 /**
  * Interface for the messaging module network extension
