[Messaging] fixed MessageRouter log

#Synced-CL 48448085

diff --git a/Engine/Source/Runtime/Messaging/Private/Bus/MessageRouter.cpp b/Engine/Source/Runtime/Messaging/Private/Bus/MessageRouter.cpp
--- a/Engine/Source/Runtime/Messaging/Private/Bus/MessageRouter.cpp
+++ b/Engine/Source/Runtime/Messaging/Private/Bus/MessageRouter.cpp
@@ -145,11 +145,20 @@ void FMessageRouter::DispatchMessage(const TSharedRef<IMessageContext, ESPMode::
 			FilterSubscriptions(ActiveSubscriptions.FindOrAdd(Context->GetMessageTypePathName()), Context, Recipients);
 			FilterSubscriptions(ActiveSubscriptions.FindOrAdd(IMessageBus::PATHNAME_All), Context, Recipients);
 
-			if (UE_GET_LOG_VERBOSITY(LogMessaging) >= ELogVerbosity::Verbose)
+			UE_SUPPRESS(LogMessaging, Verbose,
 			{
-				FString RecipientStr = FString::JoinBy(Context->GetRecipients(), TEXT("+"), &FMessageAddress::ToString);
-				UE_LOG(LogMessaging, Verbose, TEXT("Dispatching %s from %s to %s subscribers"), *Context->GetMessageTypePathName().ToString(), *Context->GetSender().ToString(), *RecipientStr);
-			}
+				const FString RecipientStr = FString::JoinBy(Recipients, TEXT("+"),
+					[](const TSharedPtr<IMessageReceiver>& Receiver)
+					{
+						return Receiver ? Receiver->GetRecipientId().ToString() : TEXT("");
+					});
+
+				UE_LOG(LogMessaging, Verbose, TEXT("Dispatching %s from %s to %i subscriber(s): %s")
+					, *Context->GetMessageTypePathName().ToString()
+					, *Context->GetSender().ToString()
+					, Recipients.Num(),
+					*RecipientStr);
+			});
 		}
 
 		// dispatch the message
